Check for object system existence before unregistering customizations.

- This fixes a collection of ASAN reported issues of memory-use-after-free.

#rnx

diff --git a/Engine/Plugins/Experimental/GameFeatures/Source/GameFeaturesEditor/Private/GameFeaturesEditorModule.cpp b/Engine/Plugins/Experimental/GameFeatures/Source/GameFeaturesEditor/Private/GameFeaturesEditorModule.cpp
--- a/Engine/Plugins/Experimental/GameFeatures/Source/GameFeaturesEditor/Private/GameFeaturesEditorModule.cpp
+++ b/Engine/Plugins/Experimental/GameFeatures/Source/GameFeaturesEditor/Private/GameFeaturesEditorModule.cpp
@@ -154,7 +154,7 @@ class FGameFeaturesEditorModule : public FDefaultModuleImpl
 	virtual void ShutdownModule() override
 	{
 		// Remove the customization
-		if (FModuleManager::Get().IsModuleLoaded("PropertyEditor"))
+		if (UObjectInitialized() && FModuleManager::Get().IsModuleLoaded("PropertyEditor"))
 		{
 			FPropertyEditorModule& PropertyModule = FModuleManager::LoadModuleChecked<FPropertyEditorModule>("PropertyEditor");
 			PropertyModule.UnregisterCustomClassLayout(UGameFeatureData::StaticClass()->GetFName());
diff --git a/Engine/Plugins/Runtime/nDisplay/Source/DisplayClusterEditor/Private/DisplayClusterEditorModule.cpp b/Engine/Plugins/Runtime/nDisplay/Source/DisplayClusterEditor/Private/DisplayClusterEditorModule.cpp
--- a/Engine/Plugins/Runtime/nDisplay/Source/DisplayClusterEditor/Private/DisplayClusterEditorModule.cpp
+++ b/Engine/Plugins/Runtime/nDisplay/Source/DisplayClusterEditor/Private/DisplayClusterEditorModule.cpp
@@ -23,8 +23,11 @@ void FDisplayClusterEditorModule::StartupModule()
 
 void FDisplayClusterEditorModule::ShutdownModule()
 {
-	UnregisterSettings();
-	UnregisterCustomizations();
+	if (UObjectInitialized())
+	{
+		UnregisterSettings();
+		UnregisterCustomizations();
+	}
 }
 
 
diff --git a/Engine/Plugins/Tests/FbxAutomationTestBuilder/Source/FbxAutomationTestBuilder/Private/FbxAutomationBuilderModule.cpp b/Engine/Plugins/Tests/FbxAutomationTestBuilder/Source/FbxAutomationTestBuilder/Private/FbxAutomationBuilderModule.cpp
--- a/Engine/Plugins/Tests/FbxAutomationTestBuilder/Source/FbxAutomationTestBuilder/Private/FbxAutomationBuilderModule.cpp
+++ b/Engine/Plugins/Tests/FbxAutomationTestBuilder/Source/FbxAutomationTestBuilder/Private/FbxAutomationBuilderModule.cpp
@@ -69,8 +69,11 @@ void FFbxAutomationBuilderModule::UnregisterTabSpawner()
 	bHasRegisteredTabSpawners = false;
 
 	//Unregister the custom detail layout
-	//FPropertyEditorModule& PropertyModule = FModuleManager::LoadModuleChecked<FPropertyEditorModule>("PropertyEditor");
-	//PropertyModule.UnregisterCustomClassLayout(UFbxAutomationBuilderView::StaticClass()->GetFName());
+	//if (UObjectInitialized())
+	//{
+	//	FPropertyEditorModule& PropertyModule = FModuleManager::LoadModuleChecked<FPropertyEditorModule>("PropertyEditor");
+	//	PropertyModule.UnregisterCustomClassLayout(UFbxAutomationBuilderView::StaticClass()->GetFName());
+	//}
 
 	FGlobalTabmanager::Get()->UnregisterNomadTabSpawner("LevelEditorFbxAutomationBuilder");
 }
diff --git a/Engine/Plugins/VirtualProduction/LevelSnapshots/Source/LevelSnapshotsEditor/Private/LevelSnapshotsEditorModule.cpp b/Engine/Plugins/VirtualProduction/LevelSnapshots/Source/LevelSnapshotsEditor/Private/LevelSnapshotsEditorModule.cpp
--- a/Engine/Plugins/VirtualProduction/LevelSnapshots/Source/LevelSnapshotsEditor/Private/LevelSnapshotsEditorModule.cpp
+++ b/Engine/Plugins/VirtualProduction/LevelSnapshots/Source/LevelSnapshotsEditor/Private/LevelSnapshotsEditorModule.cpp
@@ -63,8 +63,11 @@ void FLevelSnapshotsEditorModule::ShutdownModule()
 
 	FLevelSnapshotsEditorStyle::Shutdown();
 
-	FPropertyEditorModule& PropertyModule = FModuleManager::LoadModuleChecked<FPropertyEditorModule>("PropertyEditor");
-	PropertyModule.UnregisterCustomClassLayout(UNegatableFilter::StaticClass()->GetFName());
+	if (UObjectInitialized())
+	{
+		FPropertyEditorModule& PropertyModule = FModuleManager::LoadModuleChecked<FPropertyEditorModule>("PropertyEditor");
+		PropertyModule.UnregisterCustomClassLayout(UNegatableFilter::StaticClass()->GetFName());
+	}
 	
 	FLevelSnapshotsEditorCommands::Unregister();
 
diff --git a/Engine/Plugins/VirtualProduction/MultiUserTakes/Source/ConcertTakeRecorder/Private/ConcertTakeRecorderManager.cpp b/Engine/Plugins/VirtualProduction/MultiUserTakes/Source/ConcertTakeRecorder/Private/ConcertTakeRecorderManager.cpp
--- a/Engine/Plugins/VirtualProduction/MultiUserTakes/Source/ConcertTakeRecorder/Private/ConcertTakeRecorderManager.cpp
+++ b/Engine/Plugins/VirtualProduction/MultiUserTakes/Source/ConcertTakeRecorder/Private/ConcertTakeRecorderManager.cpp
@@ -251,8 +251,11 @@ void FConcertTakeRecorderManager::UnregisterExtensions()
 	FPropertyEditorModule* PropertyEditorModule = FModuleManager::Get().GetModulePtr<FPropertyEditorModule>("PropertyEditor");
 	if (PropertyEditorModule)
 	{
-		PropertyEditorModule->UnregisterCustomClassLayout(UConcertTakeSynchronization::StaticClass()->GetFName());
-		PropertyEditorModule->UnregisterCustomClassLayout(UConcertSessionRecordSettings::StaticClass()->GetFName());
+		if (UObjectInitialized())
+		{
+			PropertyEditorModule->UnregisterCustomClassLayout(UConcertTakeSynchronization::StaticClass()->GetFName());
+			PropertyEditorModule->UnregisterCustomClassLayout(UConcertSessionRecordSettings::StaticClass()->GetFName());
+		}
 
 		FConcertTakeRecorderSynchronizationCustomization::OnSyncPropertyValueChanged().Remove(TakeSyncDelegate);
 	}
