In the case where multiple nodes are recording, we need to have a unique
name associated with the take. This change adds support for re-naming
the take when a name conflict will occur on record.

diff --git a/Engine/Plugins/VirtualProduction/MultiUserTakes/Source/ConcertTakeRecorder/Private/ConcertTakeRecorderManager.cpp b/Engine/Plugins/VirtualProduction/MultiUserTakes/Source/ConcertTakeRecorder/Private/ConcertTakeRecorderManager.cpp
--- a/Engine/Plugins/VirtualProduction/MultiUserTakes/Source/ConcertTakeRecorder/Private/ConcertTakeRecorderManager.cpp
+++ b/Engine/Plugins/VirtualProduction/MultiUserTakes/Source/ConcertTakeRecorder/Private/ConcertTakeRecorderManager.cpp
@@ -2,8 +2,11 @@
 
 #include "ConcertTakeRecorderManager.h"
 
+#include "Algo/Count.h"
+#include "Algo/Transform.h"
 #include "ConcertMessages.h"
 #include "ConcertTransportMessages.h"
+#include "Delegates/IDelegateInstance.h"
 #include "HAL/IConsoleManager.h"
 #include "IConcertClientTransactionBridge.h"
 #include "Layout/Visibility.h"
@@ -46,6 +49,8 @@
 #include "Widgets/Images/SImage.h"
 #include "Widgets/SOverlay.h"
 #include "Components/SlateWrapperTypes.h"
+#include "Framework/Notifications/NotificationManager.h"
+#include "Widgets/Notifications/SNotificationList.h"
 
 #include "EditorFontGlyphs.h"
 #include "EditorStyleSet.h"
@@ -92,6 +97,8 @@ FConcertTakeRecorderManager::FConcertTakeRecorderManager()
 		RegisterExtensions();
 
 		UTakeRecorder::OnRecordingInitialized().AddRaw(this, &FConcertTakeRecorderManager::OnTakeRecorderInitialized);
+		FTakeRecorderParameterDelegate ParamHandler = FTakeRecorderParameterDelegate::CreateRaw(this, &FConcertTakeRecorderManager::SetupTakeParametersForMultiuser);
+		UTakeRecorder::TakeInitializeParameterOverride().RegisterHandler(TEXT("ConcertTakeRecord"), MoveTemp(ParamHandler));
 	}
 	else
 	{
@@ -102,6 +109,7 @@ FConcertTakeRecorderManager::FConcertTakeRecorderManager()
 FConcertTakeRecorderManager::~FConcertTakeRecorderManager()
 {
 	UTakeRecorder::OnRecordingInitialized().RemoveAll(this);
+	UTakeRecorder::TakeInitializeParameterOverride().UnregisterHandler(TEXT("ConcertTakeRecord"));
 
 	if (TSharedPtr<IConcertSyncClient> ConcertSyncClient = IConcertSyncClientModule::Get().GetClient(TEXT("MultiUser")))
 	{
@@ -269,17 +277,20 @@ void FConcertTakeRecorderManager::UnregisterExtensions()
 	}
 }
 
-EVisibility FConcertTakeRecorderManager::GetMultiUserIconVisibility() const
+namespace ConcertTakeManagerHelper
+{
+EVisibility UIVisFromBool(bool bValue)
+{
+	return bValue ? EVisibility::HitTestInvisible : EVisibility::Collapsed;
+}
+}
+
+bool FConcertTakeRecorderManager::ShouldIconBeVisible() const
 {
 	ITakeRecorderModule& TakeRecorderModule = FModuleManager::LoadModuleChecked<ITakeRecorderModule>("TakeRecorder");
 	UTakePreset* TakePreset = TakeRecorderModule.GetPendingTake();
 	UConcertTakeSynchronization const* TakeSync = GetDefault<UConcertTakeSynchronization>();
 
-	auto ComputeVisState = [](bool Value)
-	{
-		return Value ? EVisibility::HitTestInvisible : EVisibility::Collapsed;
-	};
-
 	if (WeakSession.IsValid() && TakePreset && !bIsRecording && TakeSync->bSyncTakeRecordingTransactions
 		&& IsTakeSyncEnabled() && CanAnyRecord())
 	{
@@ -287,11 +298,16 @@ EVisibility FConcertTakeRecorderManager::GetMultiUserIconVisibility() const
 		if (LevelSequence)
 		{
 			UTakeRecorderSources* Sources = LevelSequence->FindMetaData<UTakeRecorderSources>();
-			return ComputeVisState(Sources && Sources->GetSources().Num() > 0);
+			return Sources && Sources->GetSources().Num() > 0;
 		}
 	}
 
-	return ComputeVisState(false);
+	return false;
+}
+
+EVisibility FConcertTakeRecorderManager::GetMultiUserIconVisibility() const
+{
+	return ConcertTakeManagerHelper::UIVisFromBool(ShouldIconBeVisible());
 }
 
 void FConcertTakeRecorderManager::CreateRecordButtonOverlay(TArray<TSharedRef<SWidget>>& OutExtensions)
@@ -381,7 +397,8 @@ void FConcertTakeRecorderManager::OnTakeRecorderInitialized(UTakeRecorder* TakeR
 				Writer.SerializeObject(TakeMetaData);
 
 				InLocalIdentifierTable.GetState(TakeInitializedEvent.IdentifierState);
-				Session->SendCustomEvent(TakeInitializedEvent, Session->GetSessionClientEndpointIds(), EConcertMessageFlags::ReliableOrdered);
+				UE_LOG(LogConcertTakeRecorder, Display, TEXT("Sending custom event to initialize recording."));
+				Session->SendCustomEvent(TakeInitializedEvent, Session->GetSessionClientEndpointIds(), EConcertMessageFlags::ReliableOrdered | EConcertMessageFlags::UniqueId);
 				bIsRecording = true;
 			}
 		}
@@ -396,7 +413,7 @@ void FConcertTakeRecorderManager::OnRecordingFinished(UTakeRecorder* TakeRecorde
 		{
 			FConcertRecordingFinishedEvent RecordingFinishedEvent;
 			RecordingFinishedEvent.TakeName = TakeRecorder->GetName();
-			Session->SendCustomEvent(RecordingFinishedEvent, Session->GetSessionClientEndpointIds(), EConcertMessageFlags::ReliableOrdered);
+			Session->SendCustomEvent(RecordingFinishedEvent, Session->GetSessionClientEndpointIds(), EConcertMessageFlags::ReliableOrdered | EConcertMessageFlags::UniqueId);
 		}
 	}
 
@@ -458,6 +475,7 @@ void FConcertTakeRecorderManager::OnTakeInitializedEvent(const FConcertSessionCo
 			DefaultParams.Project = GetDefault<UTakeRecorderProjectSettings>()->Settings;
 
 			UTakeRecorder* NewRecorder = NewObject<UTakeRecorder>(GetTransientPackage(), NAME_None, RF_Transient);
+			UE_LOG(LogConcertTakeRecorder, Display, TEXT("Take initialized for recording by multi-user event."));
 			bIsRecording = true;
 			NewRecorder->Initialize(
 				LevelSequence,
@@ -750,7 +768,21 @@ void FConcertTakeRecorderManager::RecordSettingChange(const FConcertClientRecord
 		}
 
 		RecordSettings->SaveConfig();
+		const SIZE_T Recorders = CanRecord() + RemoteRecorders();
+		if (!bHaveWarned && Recorders > 1)
+		{
+			FText ErrorText = LOCTEXT("ToggleTakeRecorderWarnNameTooltip",
+									  "Multiple machines are recording. To prevent name conflicts, the multi-user name will be appended to the take path.");
+			FNotificationInfo Info(ErrorText);
+			Info.ExpireDuration = 5.0f;
+			FSlateNotificationManager::Get().AddNotification(Info)->SetCompletionState(SNotificationItem::CS_None);
+			bHaveWarned = true;
+		}
 
+		if (Recorders == 1)
+		{
+			bHaveWarned = false;
+		}
 		FConcertRecordSettingsChangeEvent OutEvent;
 		OutEvent.Settings   = RecordSetting.Settings;
 		OutEvent.EndpointId = RecordSetting.Details.ClientEndpointId;
@@ -809,4 +841,39 @@ ETransactionFilterResult FConcertTakeRecorderManager::ShouldObjectBeTransacted(U
 	return ETransactionFilterResult::IncludeObject;
 }
 
+SIZE_T FConcertTakeRecorderManager::RemoteRecorders() const
+{
+	UConcertSessionRecordSettings const* RecordSettings = GetDefault<UConcertSessionRecordSettings>();
+	SIZE_T Num = Algo::CountIf(RecordSettings->RemoteSettings,
+							   [](const FConcertClientRecordSetting&  Remote)
+							   {
+								   return Remote.Settings.bRecordOnClient;
+							   });
+	return Num;
+}
+
+FTakeRecorderParameters FConcertTakeRecorderManager::SetupTakeParametersForMultiuser(const FTakeRecorderParameters& Input)
+{
+	if (WeakSession.IsValid())
+	{
+		if (CanRecord() && RemoteRecorders() > 0)
+		{
+			TSharedPtr<IConcertClientSession> Session = WeakSession.Pin();
+			FString Name = Session->GetLocalClientInfo().DisplayName;
+			Name.RemoveSpacesInline();
+			FTakeRecorderParameters Output = Input;
+			Output.Project.TakeSaveDir = Input.Project.TakeSaveDir + "_" + Name;
+			return Output;
+		}
+		else if (!CanRecord())
+		{
+			FTakeRecorderParameters Output = Input;
+			Output.Project.TakeSaveDir = Input.Project.TakeSaveDir + "_temp";
+			return Output;
+		}
+	}
+
+	return Input;
+}
+
 #undef LOCTEXT_NAMESPACE /*ConcertTakeRecorder*/
diff --git a/Engine/Plugins/VirtualProduction/MultiUserTakes/Source/ConcertTakeRecorder/Private/ConcertTakeRecorderManager.h b/Engine/Plugins/VirtualProduction/MultiUserTakes/Source/ConcertTakeRecorder/Private/ConcertTakeRecorderManager.h
--- a/Engine/Plugins/VirtualProduction/MultiUserTakes/Source/ConcertTakeRecorder/Private/ConcertTakeRecorderManager.h
+++ b/Engine/Plugins/VirtualProduction/MultiUserTakes/Source/ConcertTakeRecorder/Private/ConcertTakeRecorderManager.h
@@ -5,6 +5,7 @@
 #include "ConcertTakeRecorderMessages.h"
 #include "ConcertMessages.h"
 #include "Delegates/IDelegateInstance.h"
+#include "HAL/Platform.h"
 #include "UObject/StrongObjectPtr.h"
 #include "TakePreset.h"
 #include "ConcertSyncClient/Public/IConcertClientTransactionBridge.h"
@@ -75,6 +76,8 @@ private:
 	void CreateExtensionWidget(TArray<TSharedRef<class SWidget>>& OutExtensions);
 	void CreateRecordButtonOverlay(TArray<TSharedRef<SWidget>>& OutExtensions);
 
+	SIZE_T RemoteRecorders() const;
+	bool ShouldIconBeVisible() const;
 	EVisibility GetMultiUserIconVisibility() const;
 
 	bool IsTakeSyncEnabled() const;
@@ -88,6 +91,8 @@ private:
 
 	ETransactionFilterResult ShouldObjectBeTransacted(UObject* InObject, UPackage* InPackage);
 private:
+	FTakeRecorderParameters SetupTakeParametersForMultiuser(const FTakeRecorderParameters& Input);
+
 	void ReportRecordingError(FText &);
 	bool CanRecord() const;
 	bool CanAnyRecord() const;
@@ -123,4 +128,7 @@ private:
 	 * Denotes if we are currently in recording mode.
 	 */
 	bool bIsRecording = false;
+
+	/** Indicates if we have warned the user about multiple recorders. */
+	bool bHaveWarned = false;
 };
