Sanitize asset name before overriding.

diff --git a/Engine/Plugins/VirtualProduction/MultiUserTakes/Source/ConcertTakeRecorder/ConcertTakeRecorder.build.cs b/Engine/Plugins/VirtualProduction/MultiUserTakes/Source/ConcertTakeRecorder/ConcertTakeRecorder.build.cs
--- a/Engine/Plugins/VirtualProduction/MultiUserTakes/Source/ConcertTakeRecorder/ConcertTakeRecorder.build.cs
+++ b/Engine/Plugins/VirtualProduction/MultiUserTakes/Source/ConcertTakeRecorder/ConcertTakeRecorder.build.cs
@@ -27,6 +27,7 @@ public class ConcertTakeRecorder : ModuleRules
 				"SlateCore",
 				"TakesCore",
 				"TakeRecorder",
+				"UnrealEd"
 			}
         );
     }
diff --git a/Engine/Plugins/VirtualProduction/MultiUserTakes/Source/ConcertTakeRecorder/Private/ConcertTakeRecorderManager.cpp b/Engine/Plugins/VirtualProduction/MultiUserTakes/Source/ConcertTakeRecorder/Private/ConcertTakeRecorderManager.cpp
--- a/Engine/Plugins/VirtualProduction/MultiUserTakes/Source/ConcertTakeRecorder/Private/ConcertTakeRecorderManager.cpp
+++ b/Engine/Plugins/VirtualProduction/MultiUserTakes/Source/ConcertTakeRecorder/Private/ConcertTakeRecorderManager.cpp
@@ -41,6 +41,7 @@
 #include "Templates/SharedPointer.h"
 #include "Templates/UnrealTemplate.h"
 #include "UObject/UObjectGlobals.h"
+#include "PackageTools.h"
 
 #include "Widgets/DeclarativeSyntaxSupport.h"
 #include "Widgets/SCompoundWidget.h"
@@ -886,7 +887,7 @@ FTakeRecorderParameters FConcertTakeRecorderManager::SetupTakeParametersForMulti
 		if (CanRecord() && RemoteRecorders() > 0)
 		{
 			TSharedPtr<IConcertClientSession> Session = WeakSession.Pin();
-			FString Name = Session->GetLocalClientInfo().DisplayName;
+			FString Name = UPackageTools::SanitizePackageName(Session->GetLocalClientInfo().DisplayName);
 			Name.RemoveSpacesInline();
 			FTakeRecorderParameters Output = Input;
 			Output.Project.TakeSaveDir = Input.Project.TakeSaveDir + "_" + Name;
