MultiUserTakes: change capitalization of "MetaData" in "bTransactTakeMetaData" property

Typo fix: Making the "D" in "Metadata" lower-case makes "Metadata" one word in the UI.

diff --git a/Engine/Plugins/VirtualProduction/MultiUserTakes/Source/ConcertTakeRecorder/Private/ConcertTakeRecorderManager.cpp b/Engine/Plugins/VirtualProduction/MultiUserTakes/Source/ConcertTakeRecorder/Private/ConcertTakeRecorderManager.cpp
--- a/Engine/Plugins/VirtualProduction/MultiUserTakes/Source/ConcertTakeRecorder/Private/ConcertTakeRecorderManager.cpp
+++ b/Engine/Plugins/VirtualProduction/MultiUserTakes/Source/ConcertTakeRecorder/Private/ConcertTakeRecorderManager.cpp
@@ -859,7 +859,7 @@ ETransactionFilterResult FConcertTakeRecorderManager::ShouldObjectBeTransacted(U
 
 	UConcertTakeSynchronization const* TakeSync	= GetDefault<UConcertTakeSynchronization>();
 	if (WeakSession.IsValid()
-		&& TakeSync->bTransactTakeMetaData
+		&& TakeSync->bTransactTakeMetadata
 		&& InObject
 		&& InObject->IsA<UTakeMetaData>())
 	{
diff --git a/Engine/Plugins/VirtualProduction/MultiUserTakes/Source/ConcertTakeRecorder/Private/ConcertTakeRecorderMessages.h b/Engine/Plugins/VirtualProduction/MultiUserTakes/Source/ConcertTakeRecorder/Private/ConcertTakeRecorderMessages.h
--- a/Engine/Plugins/VirtualProduction/MultiUserTakes/Source/ConcertTakeRecorder/Private/ConcertTakeRecorderMessages.h
+++ b/Engine/Plugins/VirtualProduction/MultiUserTakes/Source/ConcertTakeRecorder/Private/ConcertTakeRecorderMessages.h
@@ -20,7 +20,7 @@ public:
 	bool bSyncTakeRecordingTransactions = true;
 
 	UPROPERTY(config,EditAnywhere,BlueprintReadWrite,Category="Multi-user Take Synchronization")
-	bool bTransactTakeMetaData = false;
+	bool bTransactTakeMetadata = false;
 };
 
 USTRUCT()
diff --git a/Engine/Plugins/VirtualProduction/MultiUserTakes/Source/ConcertTakeRecorder/Private/ConcertTakeRecorderSynchronizationCustomization.h b/Engine/Plugins/VirtualProduction/MultiUserTakes/Source/ConcertTakeRecorder/Private/ConcertTakeRecorderSynchronizationCustomization.h
--- a/Engine/Plugins/VirtualProduction/MultiUserTakes/Source/ConcertTakeRecorder/Private/ConcertTakeRecorderSynchronizationCustomization.h
+++ b/Engine/Plugins/VirtualProduction/MultiUserTakes/Source/ConcertTakeRecorder/Private/ConcertTakeRecorderSynchronizationCustomization.h
@@ -59,7 +59,7 @@ public:
 		IDetailCategoryBuilder& TakeSync = DetailLayout.EditCategory("Multi-user Take Synchronization");
 
 		TSharedPtr<IPropertyHandle> TransactMetaData = DetailLayout.GetProperty(
-			GET_MEMBER_NAME_CHECKED(UConcertTakeSynchronization, bTransactTakeMetaData));
+			GET_MEMBER_NAME_CHECKED(UConcertTakeSynchronization, bTransactTakeMetadata));
 		TransactMetaData->SetOnChildPropertyValueChanged(FSimpleDelegate::CreateRaw(this, &ThisType::TakeMetaPropertyChanged));
 
 		TSharedPtr<IPropertyHandle> SyncTakeRecordingProperty = DetailLayout.GetProperty(
