Take recorder improvements for review last take.  This allows Multi-user to better indicate when a take is ready for review and what take to use for review.

diff --git a/Engine/Plugins/Developer/Concert/ConcertSync/ConcertSyncClient/Source/ConcertSyncClient/Private/ConcertClientSequencerManager.cpp b/Engine/Plugins/Developer/Concert/ConcertSync/ConcertSyncClient/Source/ConcertSyncClient/Private/ConcertClientSequencerManager.cpp
--- a/Engine/Plugins/Developer/Concert/ConcertSync/ConcertSyncClient/Source/ConcertSyncClient/Private/ConcertClientSequencerManager.cpp
+++ b/Engine/Plugins/Developer/Concert/ConcertSync/ConcertSyncClient/Source/ConcertSyncClient/Private/ConcertClientSequencerManager.cpp
@@ -117,9 +117,10 @@ void FConcertClientSequencerManager::OnSequencerCreated(TSharedRef<ISequencer> I
 	OpenSequencer.OnCloseEventHandle = InSequencer->OnCloseEvent().AddRaw(this, &FConcertClientSequencerManager::OnSequencerClosed);
 	int32 OpenIndex = OpenSequencers.Add(OpenSequencer);
 
-	// Setup stored state
-	InSequencer->SetPlaybackStatus((EMovieScenePlayerStatus::Type)SequencerState.PlayerStatus);
-	InSequencer->SetPlaybackSpeed(SequencerState.PlaybackSpeed);
+	// Setup stored state.  Since this is an open event by the sequencer we might be the controller and should not set player
+	// state from previous stored state.  Hence playback mode is set to undefined.  We should wait for a global time event from
+	// other clients before we set our playback state.
+	//
 	// Setting the global time will notify the server of this newly opened state.
 	InSequencer->SetGlobalTime(SequencerState.Time.ConvertTo(InSequencer->GetRootTickResolution()));
 	// Since setting the global time will potentially have set our playback mode put us back to undefined
diff --git a/Engine/Plugins/VirtualProduction/MultiUserTakes/Source/ConcertTakeRecorder/Private/ConcertTakeRecorderManager.cpp b/Engine/Plugins/VirtualProduction/MultiUserTakes/Source/ConcertTakeRecorder/Private/ConcertTakeRecorderManager.cpp
--- a/Engine/Plugins/VirtualProduction/MultiUserTakes/Source/ConcertTakeRecorder/Private/ConcertTakeRecorderManager.cpp
+++ b/Engine/Plugins/VirtualProduction/MultiUserTakes/Source/ConcertTakeRecorder/Private/ConcertTakeRecorderManager.cpp
@@ -159,7 +159,7 @@ void FConcertTakeRecorderManager::Register(TSharedRef<IConcertClientSession> InS
 		// on session start.
 		//
 		ITakeRecorderModule& TakeRecorderModule = FModuleManager::LoadModuleChecked<ITakeRecorderModule>("TakeRecorder");
-		Preset = TStrongObjectPtr<UTakePreset>(UTakePreset::AllocateTransientPreset(nullptr));
+		Preset = UTakePreset::AllocateTransientPreset(nullptr);
 		check(Preset);
 		Preset->CreateLevelSequence();
 		ULevelSequence *LevelSequence = Preset->GetLevelSequence();
@@ -178,7 +178,7 @@ void FConcertTakeRecorderManager::Register(TSharedRef<IConcertClientSession> InS
 
 	InSession->RegisterCustomEventHandler<FConcertRecordSettingsChangeEvent>(this, &FConcertTakeRecorderManager::OnRecordSettingsChangeEvent);
 	InSession->RegisterCustomEventHandler<FConcertMultiUserSyncChangeEvent>(this, &FConcertTakeRecorderManager::OnMultiUserSyncChangeEvent);
-
+	InSession->RegisterCustomEventHandler<FConcertRecordingNamedLevelSequenceEvent>(this, &FConcertTakeRecorderManager::OnNamedLevelSequenceEvent);
 	if (InSession->GetConnectionStatus() == EConcertConnectionStatus::Connected)
 	{
 		OnSessionConnectionChanged(*InSession, EConcertConnectionStatus::Connected);
@@ -205,8 +205,10 @@ void FConcertTakeRecorderManager::Unregister(TSharedRef<IConcertClientSession> I
 
 		Session->UnregisterCustomEventHandler<FConcertMultiUserSyncChangeEvent>(this);
 		Session->UnregisterCustomEventHandler<FConcertRecordSettingsChangeEvent>(this);
+
+		Session->UnregisterCustomEventHandler<FConcertRecordingNamedLevelSequenceEvent>(this);
 	}
-	Preset.Reset();
+	Preset = nullptr;
 	WeakSession.Reset();
 }
 
@@ -217,6 +219,7 @@ void FConcertTakeRecorderManager::RegisterExtensions()
 	Module.GetToolbarExtensionGenerators().AddRaw(this, &FConcertTakeRecorderManager::CreateExtensionWidget);
 	Module.GetRecordButtonExtensionGenerators().AddRaw(this, &FConcertTakeRecorderManager::CreateRecordButtonOverlay);
 	Module.GetRecordErrorCheckGenerator().AddRaw(this, &FConcertTakeRecorderManager::ReportRecordingError);
+	Module.GetCanReviewLastRecordedLevelSequenceDelegate().BindRaw(this, &FConcertTakeRecorderManager::CanReviewLastRecordedSequence);
 
 	if (GIsEditor)
 	{
@@ -253,6 +256,7 @@ void FConcertTakeRecorderManager::UnregisterExtensions()
 	ITakeRecorderModule* TakeRecorder = FModuleManager::Get().GetModulePtr<ITakeRecorderModule>("TakeRecorder");
 	if (TakeRecorder)
 	{
+		TakeRecorder->GetCanReviewLastRecordedLevelSequenceDelegate().Unbind();
 		TakeRecorder->GetToolbarExtensionGenerators().RemoveAll(this);
 		TakeRecorder->GetRecordButtonExtensionGenerators().RemoveAll(this);
 		TakeRecorder->GetRecordErrorCheckGenerator().RemoveAll(this);
@@ -367,6 +371,7 @@ void FConcertTakeRecorderManager::OnTakeRecorderInitialized(UTakeRecorder* TakeR
 	{
 		if (TSharedPtr<IConcertClientSession> Session = WeakSession.Pin())
 		{
+			LastLevelSequence = nullptr;
 			UTakeRecorderPanel* Panel = UTakeRecorderBlueprintLibrary::GetTakeRecorderPanel();
 			check(Panel);
 
@@ -418,6 +423,14 @@ void FConcertTakeRecorderManager::OnRecordingFinished(UTakeRecorder* TakeRecorde
 			FConcertRecordingFinishedEvent RecordingFinishedEvent;
 			RecordingFinishedEvent.TakeName = TakeRecorder->GetName();
 			Session->SendCustomEvent(RecordingFinishedEvent, Session->GetSessionClientEndpointIds(), EConcertMessageFlags::ReliableOrdered | EConcertMessageFlags::UniqueId);
+
+			if (CanRecord())
+			{
+				LastLevelSequence = TakeRecorder->GetSequence();
+				check(LastLevelSequence);
+				FConcertRecordingNamedLevelSequenceEvent NamedSequence{LastLevelSequence->GetPathName()};
+				Session->SendCustomEvent(NamedSequence, Session->GetSessionClientEndpointIds(), EConcertMessageFlags::ReliableOrdered);
+			}
 		}
 	}
 
@@ -465,7 +478,7 @@ void FConcertTakeRecorderManager::OnTakeInitializedEvent(const FConcertSessionCo
 		TakeRecorderState.LastStartedTake = InEvent.TakeName;
 
 		ITakeRecorderModule& TakeRecorderModule = FModuleManager::LoadModuleChecked<ITakeRecorderModule>("TakeRecorder");
-		UTakePreset* TakePreset = Preset.Get();
+		UTakePreset* TakePreset = Preset;
 		if (bConcertUseTakePresetPathForRecord > 0)
 		{
 			TakePreset = Cast<UTakePreset>(StaticLoadObject(UObject::StaticClass(), nullptr, *InEvent.TakePresetPath));
@@ -494,6 +507,9 @@ void FConcertTakeRecorderManager::OnTakeInitializedEvent(const FConcertSessionCo
 			DefaultParams.Project = GetDefault<UTakeRecorderProjectSettings>()->Settings;
 
 			UTakeRecorder* NewRecorder = NewObject<UTakeRecorder>(GetTransientPackage(), NAME_None, RF_Transient);
+
+			LastLevelSequence = nullptr;
+
 			UE_LOG(LogConcertTakeRecorder, Display, TEXT("Take initialized for recording by multi-user event."));
 			bIsRecording = true;
 			NewRecorder->Initialize(
@@ -511,8 +527,17 @@ void FConcertTakeRecorderManager::OnRecordingFinishedEvent(const FConcertSession
 	if (IsTakeSyncEnabled())
 	{
 		TakeRecorderState.LastStoppedTake = InEvent.TakeName;
+
 		if (UTakeRecorder* ActiveTakeRecorder = UTakeRecorder::GetActiveRecorder())
 		{
+			if (bIsRecording)
+			{
+				LastLevelSequence = ActiveTakeRecorder->GetSequence();
+				check(LastLevelSequence);
+				FConcertRecordingNamedLevelSequenceEvent NamedSequence{LastLevelSequence->GetPathName()};
+				TSharedPtr<IConcertClientSession> Session = WeakSession.Pin();
+				Session->SendCustomEvent(NamedSequence, Session->GetSessionClientEndpointIds(), EConcertMessageFlags::ReliableOrdered);
+			}
 			ActiveTakeRecorder->Stop();
 			bIsRecording = false;
 		}
@@ -847,7 +872,7 @@ ETransactionFilterResult FConcertTakeRecorderManager::ShouldObjectBeTransacted(U
 	UConcertSessionRecordSettings const* RecordSettings = GetDefault<UConcertSessionRecordSettings>();
 
 	ITakeRecorderModule& TakeRecorderModule = FModuleManager::LoadModuleChecked<ITakeRecorderModule>("TakeRecorder");
-	UTakePreset* TakePreset = Preset.Get();
+	UTakePreset* TakePreset = Preset;
 	if (WeakSession.IsValid()
 		&& InPackage
 		&& TakePreset
@@ -904,4 +929,27 @@ FTakeRecorderParameters FConcertTakeRecorderManager::SetupTakeParametersForMulti
 	return Input;
 }
 
+void FConcertTakeRecorderManager::OnNamedLevelSequenceEvent(const FConcertSessionContext&, const FConcertRecordingNamedLevelSequenceEvent& InEvent)
+{
+	SetLastLevelSequence(LoadObject<ULevelSequence>(nullptr, *InEvent.LevelSequencePath));
+}
+
+bool FConcertTakeRecorderManager::CanReviewLastRecordedSequence() const
+{
+	if (WeakSession.IsValid())
+	{
+		return LastLevelSequence != nullptr;
+	}
+	return true;
+}
+
+
+void FConcertTakeRecorderManager::SetLastLevelSequence(ULevelSequence* InLastSequence)
+{
+	ITakeRecorderModule& Module = FTakeRecorderRecorderManagerGetModule();
+
+	LastLevelSequence = InLastSequence;
+	Module.GetLastLevelSequenceProvider().ExecuteIfBound(LastLevelSequence);
+}
+
 #undef LOCTEXT_NAMESPACE /*ConcertTakeRecorder*/
diff --git a/Engine/Plugins/VirtualProduction/MultiUserTakes/Source/ConcertTakeRecorder/Private/ConcertTakeRecorderManager.h b/Engine/Plugins/VirtualProduction/MultiUserTakes/Source/ConcertTakeRecorder/Private/ConcertTakeRecorderManager.h
--- a/Engine/Plugins/VirtualProduction/MultiUserTakes/Source/ConcertTakeRecorder/Private/ConcertTakeRecorderManager.h
+++ b/Engine/Plugins/VirtualProduction/MultiUserTakes/Source/ConcertTakeRecorder/Private/ConcertTakeRecorderManager.h
@@ -26,7 +26,7 @@ struct EVisibility;
  * Take Recorder manager that is held by the client sync module that keeps track of when a take is started, stopped or cancelled.
  * Events are registered to client sessions that will then operate on the Take Recorder UIs
  */
-class FConcertTakeRecorderManager
+class FConcertTakeRecorderManager : public FGCObject
 {
 public:
 	/**
@@ -68,6 +68,8 @@ private:
 	void OnRecordSettingsChangeEvent(const FConcertSessionContext&, const FConcertRecordSettingsChangeEvent&);
 	void OnMultiUserSyncChangeEvent(const FConcertSessionContext&, const FConcertMultiUserSyncChangeEvent&);
 
+	void OnNamedLevelSequenceEvent(const FConcertSessionContext&, const FConcertRecordingNamedLevelSequenceEvent&);
+
 	void OnSessionClientChanged(IConcertClientSession&, EConcertClientStatus ClientStatus, const FConcertSessionClientInfo& InClientInfo);
 	void OnSessionConnectionChanged(IConcertClientSession&, EConcertConnectionStatus ConnectionStatus);
 
@@ -95,6 +97,17 @@ private:
 private:
 	FTakeRecorderParameters SetupTakeParametersForMultiuser(const FTakeRecorderParameters& Input);
 
+	virtual void AddReferencedObjects( FReferenceCollector& Collector ) override
+	{
+		Collector.AddReferencedObject(Preset);
+		Collector.AddReferencedObject(LastLevelSequence);
+	}
+
+	virtual FString GetReferencerName() const override
+	{
+		return TEXT("FConcertTakeRecorderManager");
+	}
+
 	void ReportRecordingError(FText &);
 	bool CanRecord() const;
 	bool CanAnyRecord() const;
@@ -105,6 +118,9 @@ private:
 	void RecordSettingChange(const FConcertClientRecordSetting& RecordSetting);
 	void AddRemoteClient(const FConcertSessionClientInfo& ClientInfo);
 
+	void SetLastLevelSequence(ULevelSequence* InLastSequence);
+	bool CanReviewLastRecordedSequence() const;
+
 	/** Weak pointer to the client session with which to send events. May be null or stale. */
 	TWeakPtr<IConcertClientSession> WeakSession;
 
@@ -117,7 +133,9 @@ private:
 		FString LastStoppedTake;
 	} TakeRecorderState;
 
-	TStrongObjectPtr<UTakePreset> Preset;
+	UTakePreset* Preset = nullptr;
+	ULevelSequence* LastLevelSequence = nullptr;
+
 	TSharedPtr<FConcertTakeRecorderClientSessionCustomization> Customization;
 
 	/** Delegate for any changes in client state. */
diff --git a/Engine/Plugins/VirtualProduction/MultiUserTakes/Source/ConcertTakeRecorder/Private/ConcertTakeRecorderMessages.h b/Engine/Plugins/VirtualProduction/MultiUserTakes/Source/ConcertTakeRecorder/Private/ConcertTakeRecorderMessages.h
--- a/Engine/Plugins/VirtualProduction/MultiUserTakes/Source/ConcertTakeRecorder/Private/ConcertTakeRecorderMessages.h
+++ b/Engine/Plugins/VirtualProduction/MultiUserTakes/Source/ConcertTakeRecorder/Private/ConcertTakeRecorderMessages.h
@@ -118,6 +118,15 @@ struct FConcertRecordingFinishedEvent
 	FString TakeName;
 };
 
+USTRUCT()
+struct FConcertRecordingNamedLevelSequenceEvent
+{
+	GENERATED_BODY()
+
+	UPROPERTY()
+	FString LevelSequencePath;
+};
+
 USTRUCT()
 struct FConcertRecordingCancelledEvent
 {
diff --git a/Engine/Plugins/VirtualProduction/Takes/Source/TakeRecorder/Private/TakeRecorderModule.h b/Engine/Plugins/VirtualProduction/Takes/Source/TakeRecorder/Private/TakeRecorderModule.h
--- a/Engine/Plugins/VirtualProduction/Takes/Source/TakeRecorder/Private/TakeRecorderModule.h
+++ b/Engine/Plugins/VirtualProduction/Takes/Source/TakeRecorder/Private/TakeRecorderModule.h
@@ -30,6 +30,8 @@ public:
 	FOnExternalObjectAddRemoveEvent& GetExternalObjectAddRemoveEventDelegate() override { return ExternalObjectAddRemoveEvent; }
 	FOnRecordErrorCheck&  GetRecordErrorCheckGenerator() override { return RecordErrorCheck; }
 	TArray<TWeakObjectPtr<>>& GetExternalObjects() override { return ExternalObjects; }
+	FLastRecordedLevelSequenceProvider& GetLastLevelSequenceProvider() override { return LastLevelSequenceProvider; }
+	FCanReviewLastRecordedLevelSequence& GetCanReviewLastRecordedLevelSequenceDelegate() override { return CanReviewLastRecordedSequence; };
 
 	void RegisterExternalObject(UObject* InExternalObject) override;
 	void UnregisterExternalObject(UObject* InExternalObject) override;
@@ -87,6 +89,8 @@ private:
 	FOnExtendSourcesMenuEvent SourcesMenuExtenderEvent;
 	FOnGenerateWidgetExtensions ToolbarExtensionGenerators;
 	FOnGenerateWidgetExtensions ButtonExtensionGenerators;
+	FLastRecordedLevelSequenceProvider LastLevelSequenceProvider;
+	FCanReviewLastRecordedLevelSequence CanReviewLastRecordedSequence;
 
 	FOnExternalObjectAddRemoveEvent ExternalObjectAddRemoveEvent;
 	FOnForceSaveAsPreset ForceSaveAsPresetEvent;
diff --git a/Engine/Plugins/VirtualProduction/Takes/Source/TakeRecorder/Private/Widgets/STakeRecorderPanel.cpp b/Engine/Plugins/VirtualProduction/Takes/Source/TakeRecorder/Private/Widgets/STakeRecorderPanel.cpp
--- a/Engine/Plugins/VirtualProduction/Takes/Source/TakeRecorder/Private/Widgets/STakeRecorderPanel.cpp
+++ b/Engine/Plugins/VirtualProduction/Takes/Source/TakeRecorder/Private/Widgets/STakeRecorderPanel.cpp
@@ -73,6 +73,7 @@ STakeRecorderPanel::~STakeRecorderPanel()
 	UTakeRecorder::OnRecordingInitialized().Remove(OnRecordingInitializedHandle);
 	ITakeRecorderModule& TakeRecorderModule = FModuleManager::Get().LoadModuleChecked<ITakeRecorderModule>("TakeRecorder");
 	TakeRecorderModule.OnForceSaveAsPreset().Unbind();
+	TakeRecorderModule.GetLastLevelSequenceProvider().Unbind();
 	TakeRecorderModule.GetExternalObjectAddRemoveEventDelegate().Remove(OnWidgetExternalObjectChangedHandle);
 }
 
@@ -157,6 +158,7 @@ void STakeRecorderPanel::Construct(const FArguments& InArgs)
 		TakeRecorderModule.GetExternalObjectAddRemoveEventDelegate().AddSP(this, &STakeRecorderPanel::ReconfigureExternalSettings);
 
 	TakeRecorderModule.OnForceSaveAsPreset().BindRaw(this, &STakeRecorderPanel::OnSaveAsPreset);
+	TakeRecorderModule.GetLastLevelSequenceProvider().BindRaw(this, &STakeRecorderPanel::SetLastLevelSequence);
 
 	for(TWeakObjectPtr<> Object : TakeRecorderModule.GetExternalObjects())
 	{
@@ -338,7 +340,7 @@ TSharedRef<SWidget> STakeRecorderPanel::MakeToolBar()
 				.ToolTipText(LOCTEXT("ReviewLastRecording", "Review the last recording"))
 				.ForegroundColor(FSlateColor::UseForeground())
 				.ButtonStyle(FAppStyle::Get(), "HoverHintOnly")
-				.IsEnabled_Lambda([this]() { return (LastRecordedLevelSequence != nullptr && GetTakeRecorderMode() == ETakeRecorderMode::RecordNewSequence); })
+				.IsEnabled_Lambda([this]() { return (CanReviewLastLevelSequence() && GetTakeRecorderMode() == ETakeRecorderMode::RecordNewSequence); })
 				.OnClicked(this, &STakeRecorderPanel::OnReviewLastRecording)
 				[
 					SNew(SImage)
@@ -555,6 +557,22 @@ ULevelSequence* STakeRecorderPanel::GetLastRecordedLevelSequence() const
 	return LastRecordedLevelSequence;
 }
 
+void STakeRecorderPanel::SetLastLevelSequence(ULevelSequence *InSequence)
+{
+	LastRecordedLevelSequence = InSequence;
+}
+
+bool STakeRecorderPanel::CanReviewLastLevelSequence() const
+{
+	if (LastRecordedLevelSequence == nullptr)
+	{
+		return false;
+	}
+	ITakeRecorderModule& TakeRecorderModule = FModuleManager::Get().LoadModuleChecked<ITakeRecorderModule>("TakeRecorder");
+	FCanReviewLastRecordedLevelSequence& CanReview = TakeRecorderModule.GetCanReviewLastRecordedLevelSequenceDelegate();
+	return CanReview.IsBound() && CanReview.Execute();
+}
+
 ETakeRecorderMode STakeRecorderPanel::GetTakeRecorderMode() const
 {
 	if (RecordIntoLevelSequence != nullptr)
@@ -629,6 +647,7 @@ void STakeRecorderPanel::AddReferencedObjects(FReferenceCollector& Collector)
 	Collector.AddReferencedObject(RecordIntoLevelSequence);
 	Collector.AddReferencedObject(SuppliedLevelSequence);
 	Collector.AddReferencedObject(RecordingLevelSequence);
+	Collector.AddReferencedObject(LastRecordedLevelSequence);
 }
 
 
diff --git a/Engine/Plugins/VirtualProduction/Takes/Source/TakeRecorder/Private/Widgets/STakeRecorderPanel.h b/Engine/Plugins/VirtualProduction/Takes/Source/TakeRecorder/Private/Widgets/STakeRecorderPanel.h
--- a/Engine/Plugins/VirtualProduction/Takes/Source/TakeRecorder/Private/Widgets/STakeRecorderPanel.h
+++ b/Engine/Plugins/VirtualProduction/Takes/Source/TakeRecorder/Private/Widgets/STakeRecorderPanel.h
@@ -143,6 +143,9 @@ private:
 	void SetRecordIntoLevelSequence(ULevelSequence* LevelSequence);
 
 private:
+	bool CanReviewLastLevelSequence() const;
+
+	void SetLastLevelSequence(ULevelSequence* InLastLevelSequence);
 
 	/** The transient preset that we use to - kept alive by AddReferencedObjects */
 	UTakePreset* TransientPreset;
diff --git a/Engine/Plugins/VirtualProduction/Takes/Source/TakeRecorder/Public/ITakeRecorderModule.h b/Engine/Plugins/VirtualProduction/Takes/Source/TakeRecorder/Public/ITakeRecorderModule.h
--- a/Engine/Plugins/VirtualProduction/Takes/Source/TakeRecorder/Public/ITakeRecorderModule.h
+++ b/Engine/Plugins/VirtualProduction/Takes/Source/TakeRecorder/Public/ITakeRecorderModule.h
@@ -35,6 +35,17 @@ DECLARE_MULTICAST_DELEGATE_OneParam(FOnRecordErrorCheck, FText &);
  */
 DECLARE_DELEGATE(FOnForceSaveAsPreset);
 
+class ULevelSequence;
+/**
+ * Delegate to provide last recorded level sequence.
+ */
+DECLARE_DELEGATE_OneParam(FLastRecordedLevelSequenceProvider, ULevelSequence*);
+
+/**
+ * Delegate to indicate if it is safe to review the last recorded level sequence.
+ */
+DECLARE_DELEGATE_RetVal(bool, FCanReviewLastRecordedLevelSequence);
+
 /**
  * Public module interface for the Take Recorder module
  */
@@ -75,6 +86,16 @@ public:
 	 */
 	virtual FOnForceSaveAsPreset& OnForceSaveAsPreset() = 0;
 
+	/**
+	 * Get the delegate for providing the current level sequence to the take recorder panel.
+	 */
+	virtual FLastRecordedLevelSequenceProvider& GetLastLevelSequenceProvider() = 0;
+
+	/**
+	 * Provides the take recorder panel state of review last recording.
+	 */
+	virtual FCanReviewLastRecordedLevelSequence& GetCanReviewLastRecordedLevelSequenceDelegate() = 0;
+
 	/**
 	 * Unregister a previously registered extension callback for the 'Add Source' menu
 	 */
