Make sure we have a panel before initializing a take record.

diff --git a/Engine/Plugins/VirtualProduction/MultiUserTakes/Source/ConcertTakeRecorder/Private/ConcertTakeRecorderManager.cpp b/Engine/Plugins/VirtualProduction/MultiUserTakes/Source/ConcertTakeRecorder/Private/ConcertTakeRecorderManager.cpp
--- a/Engine/Plugins/VirtualProduction/MultiUserTakes/Source/ConcertTakeRecorder/Private/ConcertTakeRecorderManager.cpp
+++ b/Engine/Plugins/VirtualProduction/MultiUserTakes/Source/ConcertTakeRecorder/Private/ConcertTakeRecorderManager.cpp
@@ -363,6 +363,12 @@ void FConcertTakeRecorderManager::CreateExtensionWidget(TArray<TSharedRef<SWidge
 
 void FConcertTakeRecorderManager::OnTakeRecorderInitialized(UTakeRecorder* TakeRecorder)
 {
+	UTakeRecorderPanel* Panel = UTakeRecorderBlueprintLibrary::GetTakeRecorderPanel();
+	if (!Panel)
+	{
+		return;
+	}
+
 	TakeRecorder->OnRecordingFinished().AddRaw(this, &FConcertTakeRecorderManager::OnRecordingFinished);
 	TakeRecorder->OnRecordingCancelled().AddRaw(this, &FConcertTakeRecorderManager::OnRecordingCancelled);
 	TakeRecorder->OnStartPlayFrameModified().AddRaw(this, &FConcertTakeRecorderManager::OnFrameAdjustment);
@@ -372,8 +378,6 @@ void FConcertTakeRecorderManager::OnTakeRecorderInitialized(UTakeRecorder* TakeR
 		if (TSharedPtr<IConcertClientSession> Session = WeakSession.Pin())
 		{
 			LastLevelSequence = nullptr;
-			UTakeRecorderPanel* Panel = UTakeRecorderBlueprintLibrary::GetTakeRecorderPanel();
-			check(Panel);
 
 			ITakeRecorderModule& TakeRecorderModule = FModuleManager::LoadModuleChecked<ITakeRecorderModule>("TakeRecorder");
 			UTakeMetaData* TakeMetaData = Panel->GetTakeMetaData();
