Removed another unneeded #include.

diff --git a/Engine/Plugins/VirtualProduction/MultiUserTakes/Source/ConcertTakeRecorder/Private/ConcertTakeRecorderClientSessionCustomization.cpp b/Engine/Plugins/VirtualProduction/MultiUserTakes/Source/ConcertTakeRecorder/Private/ConcertTakeRecorderClientSessionCustomization.cpp
--- a/Engine/Plugins/VirtualProduction/MultiUserTakes/Source/ConcertTakeRecorder/Private/ConcertTakeRecorderClientSessionCustomization.cpp
+++ b/Engine/Plugins/VirtualProduction/MultiUserTakes/Source/ConcertTakeRecorder/Private/ConcertTakeRecorderClientSessionCustomization.cpp
@@ -16,7 +16,6 @@
 #include "IConcertSyncClientModule.h"
 
 #include "Misc/AssertionMacros.h"
-#include "Oculus/LibOVRPlatform/LibOVRPlatform/include/OVR_LaunchBlockFlowResult.h"
 #include "Styling/SlateTypes.h"
 #include "Templates/SharedPointer.h"
 #include "Types/SlateEnums.h"
