Change the default for transact sources.

diff --git a/Engine/Plugins/VirtualProduction/MultiUserTakes/Source/ConcertTakeRecorder/Private/ConcertTakeRecorderMessages.h b/Engine/Plugins/VirtualProduction/MultiUserTakes/Source/ConcertTakeRecorder/Private/ConcertTakeRecorderMessages.h
--- a/Engine/Plugins/VirtualProduction/MultiUserTakes/Source/ConcertTakeRecorder/Private/ConcertTakeRecorderMessages.h
+++ b/Engine/Plugins/VirtualProduction/MultiUserTakes/Source/ConcertTakeRecorder/Private/ConcertTakeRecorderMessages.h
@@ -32,7 +32,7 @@ struct FTakeRecordSettings
 	bool bRecordOnClient = true;
 
 	UPROPERTY(config,EditAnywhere,Category="Multi-user Client Record Settings")
-	bool bTransactSources = false;
+	bool bTransactSources = true;
 };
 
 USTRUCT()
