Ensure we do not create duplicate take names by adding a random guid for collisions.  Added an ensure message on file handle creation.

diff --git a/Engine/Plugins/Developer/Concert/ConcertSync/ConcertSyncCore/Source/ConcertSyncCore/Private/ConcertSyncSessionDatabase.cpp b/Engine/Plugins/Developer/Concert/ConcertSync/ConcertSyncCore/Source/ConcertSyncCore/Private/ConcertSyncSessionDatabase.cpp
--- a/Engine/Plugins/Developer/Concert/ConcertSync/ConcertSyncCore/Source/ConcertSyncCore/Private/ConcertSyncSessionDatabase.cpp
+++ b/Engine/Plugins/Developer/Concert/ConcertSync/ConcertSyncCore/Source/ConcertSyncCore/Private/ConcertSyncSessionDatabase.cpp
@@ -2921,11 +2921,13 @@ void FConcertSyncSessionDatabase::ScheduleAsyncWrite(const FName& InPackageName,
 	SharedStream->AsyncTask = Async(EAsyncExecution::TaskGraph, [InDstPackageBlobPathname,SharedStream]()
 	{
 		TUniquePtr<FArchive> DstAr(IFileManager::Get().CreateFileWriter(*InDstPackageBlobPathname));
+		ensureMsgf(DstAr, TEXT("ConcertSyncSessionDatabase: Unable to open a file for writing. Ensure you have enough disk space and long paths are enabled."));
+
 		FConcertPackageDataStream& InPackageDataStream = SharedStream->PackageStream;
 		FMemoryReader PackageBlobAr(SharedStream->PackageData);
 		InPackageDataStream.DataAr = &PackageBlobAr;
 		SharedStream->Result = PackageDataUtil::WritePackage(InPackageDataStream, DstAr.Get());
-		return !DstAr->IsError();
+		return DstAr && !DstAr->IsError();
 	});
 }
 
diff --git a/Engine/Plugins/VirtualProduction/MultiUserTakes/Source/ConcertTakeRecorder/Private/ConcertTakeRecorderManager.cpp b/Engine/Plugins/VirtualProduction/MultiUserTakes/Source/ConcertTakeRecorder/Private/ConcertTakeRecorderManager.cpp
--- a/Engine/Plugins/VirtualProduction/MultiUserTakes/Source/ConcertTakeRecorder/Private/ConcertTakeRecorderManager.cpp
+++ b/Engine/Plugins/VirtualProduction/MultiUserTakes/Source/ConcertTakeRecorder/Private/ConcertTakeRecorderManager.cpp
@@ -14,6 +14,7 @@
 #include "ConcertTakeRecorderStyle.h"
 
 #include "ITakeRecorderModule.h"
+#include "Misc/Guid.h"
 #include "Recorder/TakeRecorder.h"
 #include "TakePreset.h"
 #include "TakeRecorderSources.h"
@@ -882,6 +883,26 @@ SIZE_T FConcertTakeRecorderManager::RemoteRecorders() const
 	return Num;
 }
 
+namespace UE::TakeRecorderManager::Private
+{
+	bool IsDuplicateRemoteName(const FString& LocalDisplayName)
+	{
+		UConcertSessionRecordSettings const* RecordSettings = GetDefault<UConcertSessionRecordSettings>();
+		SIZE_T Num = Algo::CountIf(
+			RecordSettings->RemoteSettings,
+			[LocalDisplayName](const FConcertClientRecordSetting&  Remote)
+			{
+				if (Remote.Settings.bRecordOnClient
+					&& Remote.Details.ClientInfo.DisplayName == LocalDisplayName)
+				{
+					return true;
+				}
+				return false;
+			});
+		return Num > 0;
+	}
+}
+
 FTakeRecorderParameters FConcertTakeRecorderManager::SetupTakeParametersForMultiuser(const FTakeRecorderParameters& Input)
 {
 	if (IsTakeSyncEnabled() && WeakSession.IsValid())
@@ -889,7 +910,15 @@ FTakeRecorderParameters FConcertTakeRecorderManager::SetupTakeParametersForMulti
 		if (CanRecord() && RemoteRecorders() > 0)
 		{
 			TSharedPtr<IConcertClientSession> Session = WeakSession.Pin();
-			FString Name = UPackageTools::SanitizePackageName(Session->GetLocalClientInfo().DisplayName);
+			FString LocalName = Session->GetLocalClientInfo().DisplayName;
+			if (UE::TakeRecorderManager::Private::IsDuplicateRemoteName(LocalName))
+			{
+				static FString RandomSessionId =
+					FGuid::NewGuid().ToString(EGuidFormats::Short);
+				LocalName = LocalName + "_" + RandomSessionId;
+			}
+
+			FString Name = UPackageTools::SanitizePackageName(LocalName);
 			Name.RemoveSpacesInline();
 			FTakeRecorderParameters Output = Input;
 			Output.Project.TakeSaveDir = Input.Project.TakeSaveDir + "_" + Name;
