Ensure the sync object reader does not put the UTakeMetaData into the wrong package.

diff --git a/Engine/Plugins/VirtualProduction/MultiUserTakes/Source/ConcertTakeRecorder/Private/ConcertTakeRecorderManager.cpp b/Engine/Plugins/VirtualProduction/MultiUserTakes/Source/ConcertTakeRecorder/Private/ConcertTakeRecorderManager.cpp
--- a/Engine/Plugins/VirtualProduction/MultiUserTakes/Source/ConcertTakeRecorder/Private/ConcertTakeRecorderManager.cpp
+++ b/Engine/Plugins/VirtualProduction/MultiUserTakes/Source/ConcertTakeRecorder/Private/ConcertTakeRecorderManager.cpp
@@ -378,6 +378,7 @@ void FConcertTakeRecorderManager::OnTakeRecorderInitialized(UTakeRecorder* TakeR
 				TakeInitializedEvent.TakeName = TakeRecorder->GetName();
 				TakeInitializedEvent.TakePresetPath = TakeMetaData->GetPresetOrigin()->GetPathName();
 				TakeInitializedEvent.Settings = GetDefault<UTakeRecorderUserSettings>()->Settings;
+				TakeInitializedEvent.TakeMetaDataPath = TakeMetaData->GetPathName();
 
 				FConcertLocalIdentifierTable InLocalIdentifierTable;
 				FConcertSyncObjectWriter Writer(&InLocalIdentifierTable, TakeMetaData, TakeInitializedEvent.TakeData, true, false);
@@ -474,7 +475,8 @@ void FConcertTakeRecorderManager::OnTakeInitializedEvent(const FConcertSessionCo
 			UTakeMetaData* TakeMetadata = NewObject<UTakeMetaData>(GetTransientPackage(), NAME_None, EObjectFlags::RF_Transient);
 
 			FConcertLocalIdentifierTable Table(InEvent.IdentifierState);
-			FConcertSyncObjectReader Reader(&Table, FConcertSyncWorldRemapper(), nullptr, TakeMetadata, InEvent.TakeData);
+			FConcertSyncWorldRemapper Remapper(InEvent.TakeMetaDataPath, TakeMetadata->GetPathName());
+			FConcertSyncObjectReader Reader(&Table, MoveTemp(Remapper), nullptr, TakeMetadata, InEvent.TakeData);
 			Reader.SerializeObject(TakeMetadata);
 
 			ULevelSequence* LevelSequence = TakePreset->GetLevelSequence();
diff --git a/Engine/Plugins/VirtualProduction/MultiUserTakes/Source/ConcertTakeRecorder/Private/ConcertTakeRecorderMessages.h b/Engine/Plugins/VirtualProduction/MultiUserTakes/Source/ConcertTakeRecorder/Private/ConcertTakeRecorderMessages.h
--- a/Engine/Plugins/VirtualProduction/MultiUserTakes/Source/ConcertTakeRecorder/Private/ConcertTakeRecorderMessages.h
+++ b/Engine/Plugins/VirtualProduction/MultiUserTakes/Source/ConcertTakeRecorder/Private/ConcertTakeRecorderMessages.h
@@ -100,6 +100,9 @@ struct FConcertTakeInitializedEvent
 	UPROPERTY()
 	TArray<uint8> TakeData;
 
+	UPROPERTY()
+	FString TakeMetaDataPath;
+
 	UPROPERTY()
 	FConcertLocalIdentifierState IdentifierState;
 
