Prevent a crash on connect / disconnect by making sure the active recorder is not recording after a transition.

diff --git a/Engine/Plugins/VirtualProduction/MultiUserTakes/Source/ConcertTakeRecorder/Private/ConcertTakeRecorderManager.cpp b/Engine/Plugins/VirtualProduction/MultiUserTakes/Source/ConcertTakeRecorder/Private/ConcertTakeRecorderManager.cpp
--- a/Engine/Plugins/VirtualProduction/MultiUserTakes/Source/ConcertTakeRecorder/Private/ConcertTakeRecorderManager.cpp
+++ b/Engine/Plugins/VirtualProduction/MultiUserTakes/Source/ConcertTakeRecorder/Private/ConcertTakeRecorderManager.cpp
@@ -677,7 +677,6 @@ void FConcertTakeRecorderManager::ConnectToSession(IConcertClientSession& InSess
 
 	ClientChangeDelegate = InSession.OnSessionClientChanged().AddRaw(this, &FConcertTakeRecorderManager::OnSessionClientChanged);
 
-	
 	SendInitialState(InSession);
 	UpdateSessionClientList();
 }
@@ -694,6 +693,15 @@ void FConcertTakeRecorderManager::OnSessionConnectionChanged(IConcertClientSessi
 		UE_LOG(LogConcertTakeRecorder, Display, TEXT("Multi-user Take Recorder Disconnecting from Session: %s"), *InSession.GetSessionInfo().SessionName);
 		DisconnectFromSession();
 	}
+
+	if (UTakeRecorder* ActiveTakeRecorder = UTakeRecorder::GetActiveRecorder())
+	{
+		ETakeRecorderState ActiveState = ActiveTakeRecorder->GetState();
+		if (ActiveState != ETakeRecorderState::Stopped && ActiveState != ETakeRecorderState::Cancelled)
+		{
+			ActiveTakeRecorder->Cancel();
+		}
+	}
 }
 
 void FConcertTakeRecorderManager::ReportRecordingError(FText &OutputError)
