Make sure we don't have any open trasactions after we have started playing. The animation recorder opened a transaction which caused two problems.
1. No transactions could be captured while take recorder is recording
2. A transaction produced during a stopped event marked an existing recording asset dirty on the remote client.

diff --git a/Engine/Plugins/VirtualProduction/MultiUserTakes/Source/ConcertTakeRecorder/ConcertTakeRecorder.build.cs b/Engine/Plugins/VirtualProduction/MultiUserTakes/Source/ConcertTakeRecorder/ConcertTakeRecorder.build.cs
--- a/Engine/Plugins/VirtualProduction/MultiUserTakes/Source/ConcertTakeRecorder/ConcertTakeRecorder.build.cs
+++ b/Engine/Plugins/VirtualProduction/MultiUserTakes/Source/ConcertTakeRecorder/ConcertTakeRecorder.build.cs
@@ -27,7 +27,8 @@ public class ConcertTakeRecorder : ModuleRules
 				"SlateCore",
 				"TakesCore",
 				"TakeRecorder",
-				"UnrealEd"
+				"UnrealEd",
+				"Engine"
 			}
         );
     }
diff --git a/Engine/Plugins/VirtualProduction/MultiUserTakes/Source/ConcertTakeRecorder/Private/ConcertTakeRecorderManager.cpp b/Engine/Plugins/VirtualProduction/MultiUserTakes/Source/ConcertTakeRecorder/Private/ConcertTakeRecorderManager.cpp
--- a/Engine/Plugins/VirtualProduction/MultiUserTakes/Source/ConcertTakeRecorder/Private/ConcertTakeRecorderManager.cpp
+++ b/Engine/Plugins/VirtualProduction/MultiUserTakes/Source/ConcertTakeRecorder/Private/ConcertTakeRecorderManager.cpp
@@ -5,6 +5,7 @@
 #include "Algo/Count.h"
 #include "ConcertMessages.h"
 #include "ConcertSequencerMessages.h"
+#include "Engine/Engine.h"
 #include "IConcertClientTransactionBridge.h"
 
 #include "IConcertClient.h"
@@ -335,6 +336,27 @@ void FConcertTakeRecorderManager::CreateExtensionWidget(TArray<TSharedRef<SWidge
 	}
 }
 
+
+void FConcertTakeRecorderManager::OnTakeRecorderStarted(UTakeRecorder* TakeRecorder)
+{
+#if WITH_EDITOR
+	if (GEngine)
+	{
+		GEngine->CancelTransaction(0);
+	}
+#endif
+}
+
+void FConcertTakeRecorderManager::OnTakeRecorderStopped(UTakeRecorder* TakeRecorder)
+{
+#if WITH_EDITOR
+	if (GEngine && GEngine->CanTransact())
+	{
+		ensureMsgf(GUndo == nullptr, TEXT("A transaction is pending but we are currently stopping a Take Recorder. Transactions during a Recorder Stop can cause asset locking between Multi-user clients."));
+	}
+#endif
+}
+
 void FConcertTakeRecorderManager::OnTakeRecorderInitialized(UTakeRecorder* TakeRecorder)
 {
 	UTakeRecorderPanel* Panel = UTakeRecorderBlueprintLibrary::GetTakeRecorderPanel();
@@ -345,6 +367,9 @@ void FConcertTakeRecorderManager::OnTakeRecorderInitialized(UTakeRecorder* TakeR
 
 	TakeRecorder->OnRecordingFinished().AddRaw(this, &FConcertTakeRecorderManager::OnRecordingFinished);
 	TakeRecorder->OnRecordingCancelled().AddRaw(this, &FConcertTakeRecorderManager::OnRecordingCancelled);
+	TakeRecorder->OnRecordingStarted().AddRaw(this, &FConcertTakeRecorderManager::OnTakeRecorderStarted);
+	TakeRecorder->OnRecordingStopped().AddRaw(this, &FConcertTakeRecorderManager::OnTakeRecorderStopped);
+
 	TakeRecorder->OnStartPlayFrameModified().AddRaw(this, &FConcertTakeRecorderManager::OnFrameAdjustment);
 
 	if (IsTakeSyncEnabled() && !bIsRecording && TakeRecorderState.LastStartedTake != TakeRecorder->GetName())
@@ -432,6 +457,8 @@ void FConcertTakeRecorderManager::OnRecordingCancelled(UTakeRecorder* TakeRecord
 
 	TakeRecorder->OnRecordingFinished().RemoveAll(this);
 	TakeRecorder->OnRecordingCancelled().RemoveAll(this);
+	TakeRecorder->OnRecordingStarted().RemoveAll(this);
+	TakeRecorder->OnRecordingFinished().RemoveAll(this);
 	TakeRecorder->OnStartPlayFrameModified().RemoveAll(this);
 	bIsRecording = false;
 }
diff --git a/Engine/Plugins/VirtualProduction/MultiUserTakes/Source/ConcertTakeRecorder/Private/ConcertTakeRecorderManager.h b/Engine/Plugins/VirtualProduction/MultiUserTakes/Source/ConcertTakeRecorder/Private/ConcertTakeRecorderManager.h
--- a/Engine/Plugins/VirtualProduction/MultiUserTakes/Source/ConcertTakeRecorder/Private/ConcertTakeRecorderManager.h
+++ b/Engine/Plugins/VirtualProduction/MultiUserTakes/Source/ConcertTakeRecorder/Private/ConcertTakeRecorderManager.h
@@ -68,6 +68,9 @@ private:
 	void OnTakeRecorderInitialized(UTakeRecorder* TakeRecorder);
 	void OnRecordingFinished(UTakeRecorder* TakeRecorder);
 	void OnRecordingCancelled(UTakeRecorder* TakeRecorder);
+	void OnTakeRecorderStarted(UTakeRecorder* TakeRecorder);
+	void OnTakeRecorderStopped(UTakeRecorder* TakeRecorder);
+
 	void OnFrameAdjustment(UTakeRecorder* TakeRecorder, const FFrameNumber& InPlaybackStartFrame);
 
 	//~ Concert event handlers
