Remove the requirement to block before HotReload. We defer the hot reload until the tasks have completed.  Also we eliminated the temp take from local packages by adding a filtering mechanism in C++ (similiar to the transaction equivalent).

diff --git a/Engine/Plugins/Developer/Concert/ConcertSync/ConcertSyncClient/Source/ConcertSyncClient/Private/ConcertClientPackageBridge.cpp b/Engine/Plugins/Developer/Concert/ConcertSync/ConcertSyncClient/Source/ConcertSyncClient/Private/ConcertClientPackageBridge.cpp
--- a/Engine/Plugins/Developer/Concert/ConcertSync/ConcertSyncClient/Source/ConcertSyncClient/Private/ConcertClientPackageBridge.cpp
+++ b/Engine/Plugins/Developer/Concert/ConcertSync/ConcertSyncClient/Source/ConcertSyncClient/Private/ConcertClientPackageBridge.cpp
@@ -8,6 +8,7 @@
 
 #include "Engine/World.h"
 #include "Engine/Engine.h"
+#include "IConcertClientPackageBridge.h"
 #include "UObject/ObjectMacros.h"
 #include "UObject/ObjectSaveContext.h"
 #include "UObject/Package.h"
@@ -360,4 +361,31 @@ void FConcertClientPackageBridge::HandleMapChanged(UWorld* InWorld, EMapChangeTy
 	}
 }
 
+void FConcertClientPackageBridge::RegisterPackageFilter(FName FilterName, FPackageFilterDelegate FilterHandle)
+{
+	check(PackageFilters.Find(FilterName) == nullptr);
+	PackageFilters.Add(FilterName) = MoveTemp(FilterHandle);
+}
+
+EPackageFilterResult FConcertClientPackageBridge::IsPackageFiltered(const FConcertPackageInfo& PackageInfo) const
+{
+	for (const auto& Item : PackageFilters)
+	{
+		if(Item.Value.IsBound())
+		{
+			EPackageFilterResult Result = Item.Value.Execute(PackageInfo);
+			if (Result != EPackageFilterResult::UseDefault)
+			{
+				return Result;
+			}
+		}
+	}
+	return EPackageFilterResult::UseDefault;
+}
+
+void FConcertClientPackageBridge::UnregisterPackageFilter(FName FilterName)
+{
+	PackageFilters.FindAndRemoveChecked(FilterName);
+}
+
 #undef LOCTEXT_NAMESPACE
diff --git a/Engine/Plugins/Developer/Concert/ConcertSync/ConcertSyncClient/Source/ConcertSyncClient/Private/ConcertClientPackageBridge.h b/Engine/Plugins/Developer/Concert/ConcertSync/ConcertSyncClient/Source/ConcertSyncClient/Private/ConcertClientPackageBridge.h
--- a/Engine/Plugins/Developer/Concert/ConcertSync/ConcertSyncClient/Source/ConcertSyncClient/Private/ConcertClientPackageBridge.h
+++ b/Engine/Plugins/Developer/Concert/ConcertSync/ConcertSyncClient/Source/ConcertSyncClient/Private/ConcertClientPackageBridge.h
@@ -27,6 +27,10 @@ public:
 	virtual bool& GetIgnoreLocalSaveRef() override;
 	virtual bool& GetIgnoreLocalDiscardRef() override;
 
+	virtual void RegisterPackageFilter(FName FilterName, FPackageFilterDelegate FilterHandle) override;
+	virtual void UnregisterPackageFilter(FName FilterName) override;
+	virtual EPackageFilterResult IsPackageFiltered(const FConcertPackageInfo& PackageInfo) const override;
+
 private:
 	/** Handle any deferred tasks that should be run on the game thread. */
 	void OnEndFrame();
@@ -58,6 +62,9 @@ private:
 	/** Called when a local package discard happens */
 	FOnConcertClientLocalPackageDiscarded OnLocalPackageDiscardedDelegate;
 
+	/** Map of named packages filters that can override what is included / excluded by package bridge*/
+	TMap<FName, FPackageFilterDelegate> PackageFilters;
+
 	/** Flag to ignore package change events, used when we do not want to record package changes we generate ourselves */
 	bool bIgnoreLocalSave;
 
diff --git a/Engine/Plugins/Developer/Concert/ConcertSync/ConcertSyncClient/Source/ConcertSyncClient/Private/ConcertClientPackageManager.cpp b/Engine/Plugins/Developer/Concert/ConcertSync/ConcertSyncClient/Source/ConcertSyncClient/Private/ConcertClientPackageManager.cpp
--- a/Engine/Plugins/Developer/Concert/ConcertSync/ConcertSyncClient/Source/ConcertSyncClient/Private/ConcertClientPackageManager.cpp
+++ b/Engine/Plugins/Developer/Concert/ConcertSync/ConcertSyncClient/Source/ConcertSyncClient/Private/ConcertClientPackageManager.cpp
@@ -2,6 +2,7 @@
 
 #include "ConcertClientPackageManager.h"
 
+#include "IConcertClientPackageBridge.h"
 #include "IConcertSession.h"
 #include "IConcertFileSharingService.h"
 #include "ConcertSyncClientLiveSession.h"
@@ -394,6 +395,12 @@ bool FConcertClientPackageManager::ApplyPackageFilters(const FConcertPackageInfo
 	// (We do not want to run package filtering on Disaster Recovery session and we identify it by not having the ShouldUsePackageSandbox flag for now)
 	if (EnumHasAnyFlags(LiveSession->GetSessionFlags(), EConcertSyncSessionFlags::ShouldUsePackageSandbox))
 	{
+		EPackageFilterResult Result = PackageBridge->IsPackageFiltered(InPackageInfo);
+		if (Result != EPackageFilterResult::UseDefault)
+		{
+			return Result == EPackageFilterResult::Include ? true : false;
+		}
+
 		const UConcertSyncConfig* SyncConfig = GetDefault<UConcertSyncConfig>();
 		// Ignore packages that passes the ExcludePackageClassFilters
 		if (SyncConfig->ExcludePackageClassFilters.Num() > 0 && ConcertClientPackageManagerUtil::RunPackageFilters(SyncConfig->ExcludePackageClassFilters, InPackageInfo))
@@ -638,10 +645,13 @@ bool FConcertClientPackageManager::CanHotReloadOrPurge() const
 void FConcertClientPackageManager::HotReloadPendingPackages()
 {
 	SCOPED_CONCERT_TRACE(FConcertClientPackageManager_HotReloadPendingPackages);
-	if (CanHotReloadOrPurge())
+	if (CanHotReloadOrPurge() && PackagesPendingHotReload.Num() > 0)
 	{
 		TGuardValue<bool> HotReloadGuard(bHotReloading, true);
-		LiveSession->GetSessionDatabase().FlushAsynchronousTasks();
+		if (!LiveSession->GetSessionDatabase().HasWritePackageTasksCompleted())
+		{
+			LiveSession->GetSessionDatabase().FlushAsynchronousTasks();
+		}
 		ConcertSyncClientUtil::HotReloadPackages(PackagesPendingHotReload);
 		PackagesPendingHotReload.Reset();
 	}
@@ -650,7 +660,7 @@ void FConcertClientPackageManager::HotReloadPendingPackages()
 void FConcertClientPackageManager::PurgePendingPackages()
 {
 	SCOPED_CONCERT_TRACE(FConcertClientPackageManager_PurgePendingPackages);
-	if (CanHotReloadOrPurge())
+	if (CanHotReloadOrPurge() && PackagesPendingPurge.Num() > 0)
 	{
 		ConcertSyncClientUtil::PurgePackages(PackagesPendingPurge);
 		PackagesPendingPurge.Reset();
diff --git a/Engine/Plugins/Developer/Concert/ConcertSync/ConcertSyncClient/Source/ConcertSyncClient/Private/ConcertClientWorkspace.cpp b/Engine/Plugins/Developer/Concert/ConcertSync/ConcertSyncClient/Source/ConcertSyncClient/Private/ConcertClientWorkspace.cpp
--- a/Engine/Plugins/Developer/Concert/ConcertSync/ConcertSyncClient/Source/ConcertSyncClient/Private/ConcertClientWorkspace.cpp
+++ b/Engine/Plugins/Developer/Concert/ConcertSync/ConcertSyncClient/Source/ConcertSyncClient/Private/ConcertClientWorkspace.cpp
@@ -1113,7 +1113,9 @@ void FConcertClientWorkspace::SetIgnoreOnRestoreFlagForEmittedActivities(bool bI
 
 bool FConcertClientWorkspace::CanFinalize() const
 {
-	return bFinalizeWorkspaceSyncRequested && Algo::AllOf(CanFinalizeDelegates, [](const TTuple<FName,FCanFinalizeWorkspaceDelegate>& Pair)
+	return bFinalizeWorkspaceSyncRequested
+		&& LiveSession->GetSessionDatabase().HasWritePackageTasksCompleted()
+		&& Algo::AllOf(CanFinalizeDelegates, [](const TTuple<FName,FCanFinalizeWorkspaceDelegate>& Pair)
 		{
 			if (Pair.Get<1>().IsBound())
 			{
@@ -1154,7 +1156,8 @@ bool FConcertClientWorkspace::IsReloadingPackage(FName PackageName) const
 
 bool FConcertClientWorkspace::CanProcessPendingPackages() const
 {
-	return Algo::AllOf(CanProcessPendingDelegates, [](const TTuple<FName,FCanProcessPendingPackages>& Pair)
+	return LiveSession->GetSessionDatabase().HasWritePackageTasksCompleted() &&
+		Algo::AllOf(CanProcessPendingDelegates, [](const TTuple<FName,FCanProcessPendingPackages>& Pair)
 		{
 			if (Pair.Get<1>().IsBound())
 			{
diff --git a/Engine/Plugins/Developer/Concert/ConcertSync/ConcertSyncClient/Source/ConcertSyncClient/Private/ConcertSyncClient.cpp b/Engine/Plugins/Developer/Concert/ConcertSync/ConcertSyncClient/Source/ConcertSyncClient/Private/ConcertSyncClient.cpp
--- a/Engine/Plugins/Developer/Concert/ConcertSync/ConcertSyncClient/Source/ConcertSyncClient/Private/ConcertSyncClient.cpp
+++ b/Engine/Plugins/Developer/Concert/ConcertSync/ConcertSyncClient/Source/ConcertSyncClient/Private/ConcertSyncClient.cpp
@@ -201,6 +201,11 @@ IConcertClientTransactionBridge* FConcertSyncClient::GetTransactionBridge() cons
 	return TransactionBridge;
 }
 
+IConcertClientPackageBridge* FConcertSyncClient::GetPackageBridge() const
+{
+	return PackageBridge;
+}
+
 void FConcertSyncClient::UnregisterConcertSyncHandlers(TSharedRef<IConcertClientSession> InSession)
 {
 #if WITH_EDITOR
diff --git a/Engine/Plugins/Developer/Concert/ConcertSync/ConcertSyncClient/Source/ConcertSyncClient/Private/ConcertSyncClient.h b/Engine/Plugins/Developer/Concert/ConcertSync/ConcertSyncClient/Source/ConcertSyncClient/Private/ConcertSyncClient.h
--- a/Engine/Plugins/Developer/Concert/ConcertSync/ConcertSyncClient/Source/ConcertSyncClient/Private/ConcertSyncClient.h
+++ b/Engine/Plugins/Developer/Concert/ConcertSync/ConcertSyncClient/Source/ConcertSyncClient/Private/ConcertSyncClient.h
@@ -42,6 +42,7 @@ public:
 
 	virtual void SetFileSharingService(TSharedPtr<IConcertFileSharingService> InFileSharingService) override;
 	virtual IConcertClientTransactionBridge* GetTransactionBridge() const override;
+	virtual IConcertClientPackageBridge* GetPackageBridge() const override;
 
 private:
 	void CreateWorkspace(const TSharedRef<FConcertSyncClientLiveSession>& InLiveSession);
diff --git a/Engine/Plugins/Developer/Concert/ConcertSync/ConcertSyncClient/Source/ConcertSyncClient/Public/IConcertClientPackageBridge.h b/Engine/Plugins/Developer/Concert/ConcertSync/ConcertSyncClient/Source/ConcertSyncClient/Public/IConcertClientPackageBridge.h
--- a/Engine/Plugins/Developer/Concert/ConcertSync/ConcertSyncClient/Source/ConcertSyncClient/Public/IConcertClientPackageBridge.h
+++ b/Engine/Plugins/Developer/Concert/ConcertSync/ConcertSyncClient/Source/ConcertSyncClient/Public/IConcertClientPackageBridge.h
@@ -11,6 +11,15 @@ struct FConcertPackageInfo;
 DECLARE_MULTICAST_DELEGATE_TwoParams(FOnConcertClientLocalPackageEvent, const FConcertPackageInfo&, const FString&);
 DECLARE_MULTICAST_DELEGATE_OneParam(FOnConcertClientLocalPackageDiscarded, UPackage*);
 
+enum class EPackageFilterResult : uint8
+{
+	Include,
+	Exclude,
+	UseDefault
+};
+
+DECLARE_DELEGATE_RetVal_OneParam(EPackageFilterResult, FPackageFilterDelegate, const FConcertPackageInfo&);
+
 /**
  * Bridge between the editor package events and Concert.
  * Deals with converting package update events to Concert package data.
@@ -48,6 +57,15 @@ public:
 	 */
 	virtual FOnConcertClientLocalPackageDiscarded& OnLocalPackageDiscarded() = 0;
 
+	/** Register a package filter to exclude/include certain packagins from the session. */
+	virtual void RegisterPackageFilter(FName FilterName, FPackageFilterDelegate FilterHandle) = 0;
+
+	/** Returns true if the package should be filtered from the local sandbox.  */
+	virtual EPackageFilterResult IsPackageFiltered(const FConcertPackageInfo& PackageInfo) const = 0;
+
+	/** Unregister package filter for handling package filtering. */
+	virtual void UnregisterPackageFilter(FName FilterName) = 0;
+
 protected:
 	/**
 	 * Function to access the internal bool controlling whether local saves are currently being tracked.
diff --git a/Engine/Plugins/Developer/Concert/ConcertSync/ConcertSyncClient/Source/ConcertSyncClient/Public/IConcertSyncClient.h b/Engine/Plugins/Developer/Concert/ConcertSync/ConcertSyncClient/Source/ConcertSyncClient/Public/IConcertSyncClient.h
--- a/Engine/Plugins/Developer/Concert/ConcertSync/ConcertSyncClient/Source/ConcertSyncClient/Public/IConcertSyncClient.h
+++ b/Engine/Plugins/Developer/Concert/ConcertSync/ConcertSyncClient/Source/ConcertSyncClient/Public/IConcertSyncClient.h
@@ -15,6 +15,7 @@ class IConcertClientPresenceManager;
 class IConcertClientSequencerManager;
 struct FConcertSessionClientInfo;
 class IConcertClientTransactionBridge;
+class IConcertClientPackageBridge;
 class IConcertSyncClient;
 class IConcertFileSharingService;
 
@@ -97,4 +98,7 @@ public:
 
 	/** Get the current transaction bridge for this client */
 	virtual IConcertClientTransactionBridge* GetTransactionBridge() const = 0;
+
+	/** Get the current package bridge for this client. */
+	virtual IConcertClientPackageBridge* GetPackageBridge() const = 0;
 };
diff --git a/Engine/Plugins/Developer/Concert/ConcertSync/ConcertSyncCore/Source/ConcertSyncCore/Private/ConcertSyncSessionDatabase.cpp b/Engine/Plugins/Developer/Concert/ConcertSync/ConcertSyncCore/Source/ConcertSyncCore/Private/ConcertSyncSessionDatabase.cpp
--- a/Engine/Plugins/Developer/Concert/ConcertSync/ConcertSyncCore/Source/ConcertSyncCore/Private/ConcertSyncSessionDatabase.cpp
+++ b/Engine/Plugins/Developer/Concert/ConcertSync/ConcertSyncCore/Source/ConcertSyncCore/Private/ConcertSyncSessionDatabase.cpp
@@ -12,7 +12,7 @@
 #include "SQLitePreparedStatement.h"
 #include "Serialization/MemoryReader.h"
 #include "Serialization/MemoryWriter.h"
-#include "Async/Async.h"
+#include "Tasks/Task.h"
 
 /** Contains an asynchronous representation of the FConcertPackageDataStream */
 namespace PackageDataUtil
@@ -34,8 +34,8 @@ struct FConcertPackageAsyncDataStream
 		PackageStream = {nullptr, InStream.DataSize, &PackageData};
 	}
 
-	FName								CachedPackageName;
-	TFuture<bool>						AsyncTask;
+	FName CachedPackageName;
+	UE::Tasks::TTask<bool> AsyncTask;
 	PackageDataUtil::WritePackageResult Result;
 
 	FConcertPackageDataStream   PackageStream;
@@ -2845,27 +2845,33 @@ bool FConcertSyncSessionDatabase::LoadPackage(const FString& InPackageBlobFilena
 
 void FConcertSyncSessionDatabase::FlushAsynchronousTasks()
 {
-	for (const auto& Item : DeferredLargePackageIOPtr->GetMap())
+	while (!DeferredLargePackageIOPtr->GetMap().IsEmpty())
 	{
-		if (Item.Value->AsyncTask.Get())
-		{
-			HandleWritePackageResult(Item.Value->CachedPackageName, Item.Key, PackageFileCache, MoveTemp(Item.Value->Result));
-		}
-		else
-		{
-			UE_LOG(LogConcert, Error, TEXT("Async task failed to write package %s"), *Item.Value->CachedPackageName.ToString());
-		}
+		UpdateAsynchronousTasks();
+		const float SleepTimeInSeconds = 0.010f;
+		FPlatformProcess::Sleep(SleepTimeInSeconds);
 	}
-	DeferredLargePackageIOPtr->Clear();
+}
+
+bool FConcertSyncSessionDatabase::HasWritePackageTasksCompleted() const
+{
+	return DeferredLargePackageIOPtr->GetMap().IsEmpty();
 }
 
 void FConcertSyncSessionDatabase::UpdateAsynchronousTasks()
 {
 	for (auto It = DeferredLargePackageIOPtr->GetMap().CreateIterator(); It; ++It)
 	{
-		if (It->Value->AsyncTask.IsReady())
+		if (It->Value->AsyncTask.IsCompleted())
 		{
-			HandleWritePackageResult(It->Value->CachedPackageName, It->Key, PackageFileCache, MoveTemp(It->Value->Result));
+			if (It->Value->AsyncTask.GetResult())
+			{
+				HandleWritePackageResult(It->Value->CachedPackageName, It->Key, PackageFileCache, MoveTemp(It->Value->Result));
+			}
+			else
+			{
+				UE_LOG(LogConcert, Error, TEXT("Async task failed to write package %s"), *It->Value->CachedPackageName.ToString());
+			}
 			It.RemoveCurrent();
 		}
 	}
@@ -2914,11 +2920,13 @@ void FConcertSyncSessionDatabase::ScheduleAsyncWrite(const FName& InPackageName,
 	FConcertPackageAsyncDataStreamPtr SharedStream = MakeShared<FConcertPackageAsyncDataStream,ESPMode::ThreadSafe>(InPackageDataStream, InPackageName);
 	if(FConcertPackageAsyncDataStreamPtr* Item = DeferredLargePackageIOPtr->GetMap().Find(InDstPackageBlobPathname))
 	{
-		(*Item)->AsyncTask.Get();
-		HandleWritePackageResult((*Item)->CachedPackageName, InDstPackageBlobPathname, PackageFileCache, MoveTemp((*Item)->Result));
+		if ((*Item)->AsyncTask.GetResult())
+		{
+			HandleWritePackageResult((*Item)->CachedPackageName, InDstPackageBlobPathname, PackageFileCache, MoveTemp((*Item)->Result));
+		}
 	}
 	DeferredLargePackageIOPtr->GetMap().FindOrAdd(InDstPackageBlobPathname) = SharedStream;
-	SharedStream->AsyncTask = Async(EAsyncExecution::TaskGraph, [InDstPackageBlobPathname,SharedStream]()
+	SharedStream->AsyncTask = UE::Tasks::Launch(TEXT("ConcertWritePackageAsync"), [InDstPackageBlobPathname,SharedStream]()
 	{
 		TUniquePtr<FArchive> DstAr(IFileManager::Get().CreateFileWriter(*InDstPackageBlobPathname));
 		ensureMsgf(DstAr, TEXT("ConcertSyncSessionDatabase: Unable to open a file for writing. Ensure you have enough disk space and long paths are enabled."));
diff --git a/Engine/Plugins/Developer/Concert/ConcertSync/ConcertSyncCore/Source/ConcertSyncCore/Public/ConcertSyncSessionDatabase.h b/Engine/Plugins/Developer/Concert/ConcertSync/ConcertSyncCore/Source/ConcertSyncCore/Public/ConcertSyncSessionDatabase.h
--- a/Engine/Plugins/Developer/Concert/ConcertSync/ConcertSyncCore/Source/ConcertSyncCore/Public/ConcertSyncSessionDatabase.h
+++ b/Engine/Plugins/Developer/Concert/ConcertSync/ConcertSyncCore/Source/ConcertSyncCore/Public/ConcertSyncSessionDatabase.h
@@ -705,6 +705,11 @@ public:
 	 */
 	void FlushAsynchronousTasks();
 
+	/**
+	 * Has all asynchronous package tasks completed.
+	 */
+	bool HasWritePackageTasksCompleted() const;
+
 	FOnActivityProduced& OnActivityProduced() { return ActivityProducedEvent; }
 	
 private:
diff --git a/Engine/Plugins/VirtualProduction/MultiUserTakes/Source/ConcertTakeRecorder/Private/ConcertTakeRecorderManager.cpp b/Engine/Plugins/VirtualProduction/MultiUserTakes/Source/ConcertTakeRecorder/Private/ConcertTakeRecorderManager.cpp
--- a/Engine/Plugins/VirtualProduction/MultiUserTakes/Source/ConcertTakeRecorder/Private/ConcertTakeRecorderManager.cpp
+++ b/Engine/Plugins/VirtualProduction/MultiUserTakes/Source/ConcertTakeRecorder/Private/ConcertTakeRecorderManager.cpp
@@ -5,9 +5,10 @@
 #include "Algo/Count.h"
 #include "ConcertMessages.h"
 #include "ConcertSequencerMessages.h"
+#include "ConcertWorkspaceData.h"
 #include "Engine/Engine.h"
 #include "IConcertClientTransactionBridge.h"
-
+#include "IConcertClientPackageBridge.h"
 #include "IConcertClient.h"
 #include "IConcertSyncClient.h"
 #include "IConcertSyncClientModule.h"
@@ -200,10 +201,14 @@ void FConcertTakeRecorderManager::RegisterExtensions()
 	{
 		if (TSharedPtr<IConcertSyncClient> ConcertSyncClient = IConcertSyncClientModule::Get().GetClient(TEXT("MultiUser")))
 		{
-			IConcertClientTransactionBridge* Bridge = ConcertSyncClient->GetTransactionBridge();
-			check(Bridge != nullptr);
+			IConcertClientTransactionBridge* TransactionBridge = ConcertSyncClient->GetTransactionBridge();
+			check(TransactionBridge != nullptr);
+
+			TransactionBridge->RegisterTransactionFilter(TEXT("ConcertTakes"), FTransactionFilterDelegate::CreateRaw(this, &FConcertTakeRecorderManager::ShouldObjectBeTransacted));
 
-			Bridge->RegisterTransactionFilter(TEXT("ConcertTakes"), FTransactionFilterDelegate::CreateRaw(this, &FConcertTakeRecorderManager::ShouldObjectBeTransacted));
+			IConcertClientPackageBridge* PackageBridge = ConcertSyncClient->GetPackageBridge();
+			check(PackageBridge);
+			PackageBridge->RegisterPackageFilter(TEXT("ConcertTakes"), FPackageFilterDelegate::CreateRaw(this, &FConcertTakeRecorderManager::ShouldPackageBeFiltered));
 		}
 
 		FPropertyEditorModule& PropertyEditorModule = FModuleManager::Get().LoadModuleChecked<FPropertyEditorModule>("PropertyEditor");
@@ -251,10 +256,13 @@ void FConcertTakeRecorderManager::UnregisterExtensions()
 
 	if (TSharedPtr<IConcertSyncClient> ConcertSyncClient = IConcertSyncClientModule::Get().GetClient(TEXT("MultiUser")))
 	{
-		IConcertClientTransactionBridge* Bridge = ConcertSyncClient->GetTransactionBridge();
-		check(Bridge != nullptr);
+		IConcertClientTransactionBridge* TransactionBridge = ConcertSyncClient->GetTransactionBridge();
+		check(TransactionBridge != nullptr);
+		TransactionBridge->UnregisterTransactionFilter(TEXT("ConcertTakes"));
 
-		Bridge->UnregisterTransactionFilter(TEXT("ConcertTakes"));
+		IConcertClientPackageBridge* PackageBridge = ConcertSyncClient->GetPackageBridge();
+		check(PackageBridge != nullptr);
+		PackageBridge->UnregisterPackageFilter(TEXT("ConcertTakes"));
 	}
 }
 
@@ -882,6 +890,43 @@ void FConcertTakeRecorderManager::UpdateSessionClientList()
 	RecordSettings->SaveConfig();
 }
 
+namespace UE::TakeRecorderManager::Private
+{
+	bool IsDuplicateRemoteName(const FString& LocalDisplayName)
+	{
+		UConcertSessionRecordSettings const* RecordSettings = GetDefault<UConcertSessionRecordSettings>();
+		SIZE_T Num = Algo::CountIf(
+			RecordSettings->RemoteSettings,
+			[LocalDisplayName](const FConcertClientRecordSetting&  Remote)
+			{
+				if (Remote.Settings.bRecordOnClient
+					&& Remote.Details.ClientInfo.DisplayName == LocalDisplayName)
+				{
+					return true;
+				}
+				return false;
+			});
+		return Num > 0;
+	}
+	const FString TempPostfix = TEXT("_temp");
+}
+
+
+EPackageFilterResult FConcertTakeRecorderManager::ShouldPackageBeFiltered(const FConcertPackageInfo& InPackageInfo)
+{
+	if (IsTakeSyncEnabled() && WeakSession.IsValid() && !CanRecord())
+	{
+		FTakeRecorderProjectParameters Project = GetDefault<UTakeRecorderProjectSettings>()->Settings;
+		FString FullName = InPackageInfo.PackageName.ToString();
+		if (FullName.StartsWith(Project.RootTakeSaveDir.Path) &&
+			FullName.EndsWith(UE::TakeRecorderManager::Private::TempPostfix))
+		{
+			return EPackageFilterResult::Exclude;
+		}
+	}
+	return EPackageFilterResult::UseDefault;
+}
+
 ETransactionFilterResult FConcertTakeRecorderManager::ShouldObjectBeTransacted(UObject* InObject, UPackage* InPackage)
 {
 	UConcertSessionRecordSettings const* RecordSettings = GetDefault<UConcertSessionRecordSettings>();
@@ -920,26 +965,6 @@ SIZE_T FConcertTakeRecorderManager::RemoteRecorders() const
 	return Num;
 }
 
-namespace UE::TakeRecorderManager::Private
-{
-	bool IsDuplicateRemoteName(const FString& LocalDisplayName)
-	{
-		UConcertSessionRecordSettings const* RecordSettings = GetDefault<UConcertSessionRecordSettings>();
-		SIZE_T Num = Algo::CountIf(
-			RecordSettings->RemoteSettings,
-			[LocalDisplayName](const FConcertClientRecordSetting&  Remote)
-			{
-				if (Remote.Settings.bRecordOnClient
-					&& Remote.Details.ClientInfo.DisplayName == LocalDisplayName)
-				{
-					return true;
-				}
-				return false;
-			});
-		return Num > 0;
-	}
-}
-
 FTakeRecorderParameters FConcertTakeRecorderManager::SetupTakeParametersForMultiuser(const FTakeRecorderParameters& Input)
 {
 	if (IsTakeSyncEnabled() && WeakSession.IsValid())
@@ -964,7 +989,7 @@ FTakeRecorderParameters FConcertTakeRecorderManager::SetupTakeParametersForMulti
 		else if (!CanRecord())
 		{
 			FTakeRecorderParameters Output = Input;
-			Output.Project.TakeSaveDir = Input.Project.TakeSaveDir + "_temp";
+			Output.Project.TakeSaveDir = Input.Project.TakeSaveDir + UE::TakeRecorderManager::Private::TempPostfix;
 			return Output;
 		}
 	}
diff --git a/Engine/Plugins/VirtualProduction/MultiUserTakes/Source/ConcertTakeRecorder/Private/ConcertTakeRecorderManager.h b/Engine/Plugins/VirtualProduction/MultiUserTakes/Source/ConcertTakeRecorder/Private/ConcertTakeRecorderManager.h
--- a/Engine/Plugins/VirtualProduction/MultiUserTakes/Source/ConcertTakeRecorder/Private/ConcertTakeRecorderManager.h
+++ b/Engine/Plugins/VirtualProduction/MultiUserTakes/Source/ConcertTakeRecorder/Private/ConcertTakeRecorderManager.h
@@ -12,6 +12,7 @@ class UTakePreset;
 enum class EConcertClientStatus : uint8;
 enum class EConcertConnectionStatus : uint8;
 enum class ETransactionFilterResult : uint8;
+enum class EPackageFilterResult : uint8;
 struct FConcertClientRecordSetting;
 struct FConcertMultiUserSyncChangeEvent;
 struct FConcertRecordSettingsChangeEvent;
@@ -20,6 +21,7 @@ struct FConcertRecordingFinishedEvent;
 struct FConcertRecordingNamedLevelSequenceEvent;
 struct FConcertSessionClientInfo;
 struct FConcertTakeInitializedEvent;
+struct FConcertPackageInfo;
 struct FFrameNumber;
 struct FTakeRecorderParameters;
 
@@ -107,6 +109,8 @@ private:
 	void ConnectToSession(IConcertClientSession&);
 
 	ETransactionFilterResult ShouldObjectBeTransacted(UObject* InObject, UPackage* InPackage);
+	EPackageFilterResult ShouldPackageBeFiltered(const FConcertPackageInfo& InPackage);
+
 private:
 	FTakeRecorderParameters SetupTakeParametersForMultiuser(const FTakeRecorderParameters& Input);
 
