Change the ensure to a warning that is displayed in a Multi-user session.

diff --git a/Engine/Plugins/VirtualProduction/MultiUserTakes/Source/ConcertTakeRecorder/Private/ConcertTakeRecorderManager.cpp b/Engine/Plugins/VirtualProduction/MultiUserTakes/Source/ConcertTakeRecorder/Private/ConcertTakeRecorderManager.cpp
--- a/Engine/Plugins/VirtualProduction/MultiUserTakes/Source/ConcertTakeRecorder/Private/ConcertTakeRecorderManager.cpp
+++ b/Engine/Plugins/VirtualProduction/MultiUserTakes/Source/ConcertTakeRecorder/Private/ConcertTakeRecorderManager.cpp
@@ -358,9 +358,9 @@ void FConcertTakeRecorderManager::OnTakeRecorderStarted(UTakeRecorder* TakeRecor
 void FConcertTakeRecorderManager::OnTakeRecorderStopped(UTakeRecorder* TakeRecorder)
 {
 #if WITH_EDITOR
-	if (GEngine && GEngine->CanTransact())
+	if (GEngine && GEngine->CanTransact() && WeakSession.IsValid())
 	{
-		ensureMsgf(GUndo == nullptr, TEXT("A transaction is pending but we are currently stopping a Take Recorder. Transactions during a Recorder Stop can cause asset locking between Multi-user clients."));
+		UE_LOG(LogConcertTakeRecorder, Warning, TEXT("A transaction is pending but we are currently stopping a Take Recorder. Transactions during a Recorder Stop can cause asset locking between Multi-user clients."));
 	}
 #endif
 }
