Only unregister if we are running in editor - this matches with the registration functions in the constructor of the class.

diff --git a/Engine/Plugins/VirtualProduction/MultiUserTakes/Source/ConcertTakeRecorder/Private/ConcertTakeRecorderManager.cpp b/Engine/Plugins/VirtualProduction/MultiUserTakes/Source/ConcertTakeRecorder/Private/ConcertTakeRecorderManager.cpp
--- a/Engine/Plugins/VirtualProduction/MultiUserTakes/Source/ConcertTakeRecorder/Private/ConcertTakeRecorderManager.cpp
+++ b/Engine/Plugins/VirtualProduction/MultiUserTakes/Source/ConcertTakeRecorder/Private/ConcertTakeRecorderManager.cpp
@@ -242,27 +242,30 @@ void FConcertTakeRecorderManager::UnregisterExtensions()
 		TakeRecorder->GetRecordErrorCheckGenerator().RemoveAll(this);
 	}
 
-	FPropertyEditorModule* PropertyEditorModule = FModuleManager::Get().GetModulePtr<FPropertyEditorModule>("PropertyEditor");
-	if (PropertyEditorModule)
+	if (GIsEditor)
 	{
-		if (UObjectInitialized())
+		FPropertyEditorModule* PropertyEditorModule = FModuleManager::Get().GetModulePtr<FPropertyEditorModule>("PropertyEditor");
+		if (PropertyEditorModule)
 		{
-			PropertyEditorModule->UnregisterCustomClassLayout(UConcertTakeSynchronization::StaticClass()->GetFName());
-			PropertyEditorModule->UnregisterCustomClassLayout(UConcertSessionRecordSettings::StaticClass()->GetFName());
-		}
+			if (UObjectInitialized())
+			{
+				PropertyEditorModule->UnregisterCustomClassLayout(UConcertTakeSynchronization::StaticClass()->GetFName());
+				PropertyEditorModule->UnregisterCustomClassLayout(UConcertSessionRecordSettings::StaticClass()->GetFName());
+			}
 
-		FConcertTakeRecorderSynchronizationCustomization::OnSyncPropertyValueChanged().Remove(TakeSyncDelegate);
-	}
+			FConcertTakeRecorderSynchronizationCustomization::OnSyncPropertyValueChanged().Remove(TakeSyncDelegate);
+		}
 
-	if (TSharedPtr<IConcertSyncClient> ConcertSyncClient = IConcertSyncClientModule::Get().GetClient(TEXT("MultiUser")))
-	{
-		IConcertClientTransactionBridge* TransactionBridge = ConcertSyncClient->GetTransactionBridge();
-		check(TransactionBridge != nullptr);
-		TransactionBridge->UnregisterTransactionFilter(TEXT("ConcertTakes"));
+		if (TSharedPtr<IConcertSyncClient> ConcertSyncClient = IConcertSyncClientModule::Get().GetClient(TEXT("MultiUser")))
+		{
+			IConcertClientTransactionBridge* TransactionBridge = ConcertSyncClient->GetTransactionBridge();
+			check(TransactionBridge != nullptr);
+			TransactionBridge->UnregisterTransactionFilter(TEXT("ConcertTakes"));
 
-		IConcertClientPackageBridge* PackageBridge = ConcertSyncClient->GetPackageBridge();
-		check(PackageBridge != nullptr);
-		PackageBridge->UnregisterPackageFilter(TEXT("ConcertTakes"));
+			IConcertClientPackageBridge* PackageBridge = ConcertSyncClient->GetPackageBridge();
+			check(PackageBridge != nullptr);
+			PackageBridge->UnregisterPackageFilter(TEXT("ConcertTakes"));
+		}
 	}
 }
 
