Expose way for Blueprints to query RecordOnClient and SynchronizeTakeRecorderTransactions located at the bottom of the Take Recorder tab.
This is needed to unblock VCam HUD from locally recording takes via the HUD.

diff --git a/Engine/Plugins/VirtualProduction/MultiUserTakes/Source/ConcertTakeRecorder/Private/ConcertTakeRecorderClientSessionCustomization.cpp b/Engine/Plugins/VirtualProduction/MultiUserTakes/Source/ConcertTakeRecorder/Private/ConcertTakeRecorderClientSessionCustomization.cpp
--- a/Engine/Plugins/VirtualProduction/MultiUserTakes/Source/ConcertTakeRecorder/Private/ConcertTakeRecorderClientSessionCustomization.cpp
+++ b/Engine/Plugins/VirtualProduction/MultiUserTakes/Source/ConcertTakeRecorder/Private/ConcertTakeRecorderClientSessionCustomization.cpp
@@ -348,14 +348,17 @@ void FConcertTakeRecorderClientSessionCustomization::PopulateClientList()
 	}
 }
 
-void FConcertTakeRecorderClientSessionCustomization::UpdateClientSettings(
-	EConcertClientStatus Status, const FConcertClientRecordSetting& RecordSetting)
+const FConcertClientRecordSetting* FConcertTakeRecorderClientSessionCustomization::FindClientSettings(const FGuid& EndpointSettings) const
 {
-	int32 Index = Clients.IndexOfByPredicate([&RecordSetting](const TSharedPtr<FConcertClientRecordSetting>& Setting) {
-		return Setting->Details.ClientEndpointId == RecordSetting.Details.ClientEndpointId;
-	});
+	const int32 Index = InternalFindClientSettingsIndex(EndpointSettings);
+	return Clients.IsValidIndex(Index) ? Clients[Index].Get() : nullptr;
+}
+
+void FConcertTakeRecorderClientSessionCustomization::UpdateClientSettings(EConcertClientStatus Status, const FConcertClientRecordSetting& RecordSetting)
+{
+	const int32 Index = InternalFindClientSettingsIndex(RecordSetting.Details.ClientEndpointId);
 
-	switch(Status)
+	switch (Status)
 	{
 	case EConcertClientStatus::Connected:
 		Clients.Emplace(MakeShared<FConcertClientRecordSetting>(RecordSetting));
@@ -369,11 +372,12 @@ void FConcertTakeRecorderClientSessionCustomization::UpdateClientSettings(
 		check(Index != INDEX_NONE);
 		Clients.RemoveAt(Index);
 		break;
-	};
+	default: checkNoEntry();
+	}
 
 	if (ClientsListViewWeak.IsValid())
 	{
-		TSharedPtr<SListViewClientRecordSetting> ClientsListView = ClientsListViewWeak.Pin();
+		const TSharedPtr<SListViewClientRecordSetting> ClientsListView = ClientsListViewWeak.Pin();
 		if (Status == EConcertClientStatus::Disconnected
 			|| Status == EConcertClientStatus::Connected )
 		{
@@ -381,7 +385,7 @@ void FConcertTakeRecorderClientSessionCustomization::UpdateClientSettings(
 		}
 		else
 		{
-			TSharedPtr<STakeRecordDetailsRow> Row = StaticCastSharedPtr<STakeRecordDetailsRow>(ClientsListView->WidgetFromItem(Clients[Index]));
+			const TSharedPtr<STakeRecordDetailsRow> Row = StaticCastSharedPtr<STakeRecordDetailsRow>(ClientsListView->WidgetFromItem(Clients[Index]));
 			if (Row.IsValid())
 			{
 				Row->UpdateUI();
@@ -395,6 +399,14 @@ void FConcertTakeRecorderClientSessionCustomization::RecordSettingChange(const F
 	OnRecordSettingDelegate.Broadcast(RecordSetting);
 }
 
+int32 FConcertTakeRecorderClientSessionCustomization::InternalFindClientSettingsIndex(const FGuid& ClientEndpointId) const
+{
+	return Clients.IndexOfByPredicate([&ClientEndpointId](const TSharedPtr<FConcertClientRecordSetting>& Setting)
+		{
+			return Setting->Details.ClientEndpointId == ClientEndpointId;
+		});
+}
+
 void FConcertTakeRecorderClientSessionCustomization::CustomizeDetails(IDetailLayoutBuilder& DetailLayout)
 {
 	IDetailCategoryBuilder& RecordSettings = DetailLayout.EditCategory(TEXT("Multi-user Client Record Settings"));
@@ -441,6 +453,6 @@ void FConcertTakeRecorderClientSessionCustomization::CustomizeDetails(IDetailLay
 			)
 		];
 	}
-};
+}
 
 #undef LOCTEXT_NAMESPACE
diff --git a/Engine/Plugins/VirtualProduction/MultiUserTakes/Source/ConcertTakeRecorder/Private/ConcertTakeRecorderClientSessionCustomization.h b/Engine/Plugins/VirtualProduction/MultiUserTakes/Source/ConcertTakeRecorder/Private/ConcertTakeRecorderClientSessionCustomization.h
--- a/Engine/Plugins/VirtualProduction/MultiUserTakes/Source/ConcertTakeRecorder/Private/ConcertTakeRecorderClientSessionCustomization.h
+++ b/Engine/Plugins/VirtualProduction/MultiUserTakes/Source/ConcertTakeRecorder/Private/ConcertTakeRecorderClientSessionCustomization.h
@@ -5,33 +5,42 @@
 #include "Delegates/Delegate.h"
 #include "IDetailCustomization.h"
 
-
 enum class EConcertClientStatus : uint8;
 struct FConcertClientRecordSetting;
 template <typename ItemType> class SListView;
 
 DECLARE_MULTICAST_DELEGATE_OneParam(FOnConcertRecordSettingChanged, const FConcertClientRecordSetting&);
 
-using SListViewClientRecordSetting = SListView<TSharedPtr<FConcertClientRecordSetting>>;
-
 class FConcertTakeRecorderClientSessionCustomization : public IDetailCustomization
 {
+	using SListViewClientRecordSetting = SListView<TSharedPtr<FConcertClientRecordSetting>>;
 public:
-	void CustomizeDetails(IDetailLayoutBuilder& DetailLayout) override;
 
+	//~ Begin IDetailCustomization Interface
+	virtual void CustomizeDetails(IDetailLayoutBuilder& DetailLayout) override;
+	//~ End IDetailCustomization Interface
+
+	/** Updates the internal list of client data. Called whenever an external change is made (e.g. from remote client). */
 	void UpdateClientSettings(EConcertClientStatus Status, const FConcertClientRecordSetting& RecordSetting);
+	/** Resets the Clients list. Creates an entry for every client in the session. */
 	void PopulateClientList();
 
-	FOnConcertRecordSettingChanged& OnRecordSettingChanged()
-	{
-		return OnRecordSettingDelegate;
-	}
+	/** @return The settings saved for the client or nullptr */
+	const FConcertClientRecordSetting* FindClientSettings(const FGuid& EndpointSettings) const;
 
-private:
-	void RecordSettingChange(const FConcertClientRecordSetting &RecordSetting);
+	FOnConcertRecordSettingChanged& OnRecordSettingChanged() { return OnRecordSettingDelegate; }
 
-	FOnConcertRecordSettingChanged	OnRecordSettingDelegate;
+private:
+	
+	FOnConcertRecordSettingChanged OnRecordSettingDelegate;
 
+	/** State of the local and remote clients. Updated externally by UpdateClientSettings. Remote clients trigger networked events when they change the settings. */
 	TArray<TSharedPtr<FConcertClientRecordSetting>> Clients;
 	TWeakPtr<SListViewClientRecordSetting>		    ClientsListViewWeak;
+
+	/** Called when a setting is changed by the customization */
+	void RecordSettingChange(const FConcertClientRecordSetting &RecordSetting);
+	
+	/** @return The settings saved for the client or nullptr */
+	int32 InternalFindClientSettingsIndex(const FGuid& ClientEndpointId) const;
 };
diff --git a/Engine/Plugins/VirtualProduction/MultiUserTakes/Source/ConcertTakeRecorder/Private/ConcertTakeRecorderManager.cpp b/Engine/Plugins/VirtualProduction/MultiUserTakes/Source/ConcertTakeRecorder/Private/ConcertTakeRecorderManager.cpp
--- a/Engine/Plugins/VirtualProduction/MultiUserTakes/Source/ConcertTakeRecorder/Private/ConcertTakeRecorderManager.cpp
+++ b/Engine/Plugins/VirtualProduction/MultiUserTakes/Source/ConcertTakeRecorder/Private/ConcertTakeRecorderManager.cpp
@@ -188,6 +188,45 @@ void FConcertTakeRecorderManager::Unregister(TSharedRef<IConcertClientSession> I
 	WeakSession.Reset();
 }
 
+const FConcertClientRecordSetting* FConcertTakeRecorderManager::FindClientRecorderSetting(const FGuid& EndpointId) const
+{
+	return Customization ? Customization->FindClientSettings(EndpointId) : nullptr;
+}
+
+bool FConcertTakeRecorderManager::EditClientSettings(
+	const FGuid& EndpointId,
+	TFunctionRef<void(FTakeRecordSettings& Settings)> ModifierFunc,
+	TOptional<TFunctionRef<bool(const FTakeRecordSettings& Settings)>> Predicate
+	)
+{
+	const FConcertClientRecordSetting* Setting = FindClientRecorderSetting(EndpointId);
+	if (!Setting)
+	{
+		return false;
+	}
+
+	const bool bFailedPredicate = Predicate && !(*Predicate)(Setting->Settings);
+	if (bFailedPredicate)
+	{
+		return true;
+	}
+
+	// This unnecessarily copies a bunch of strings in FConcertSessionClientInfo as well but it keeps this logic simple...
+	// Should refactor FConcertClientRecordSetting to contain a sub-struct that is then copied.
+	FConcertClientRecordSetting SettingCopy = *Setting;
+	ModifierFunc(SettingCopy.Settings);
+
+	// Update the other clients...
+	RecordSettingChange(SettingCopy);
+	// ... and the local UI
+	if (Customization)
+	{
+		Customization->UpdateClientSettings(EConcertClientStatus::Updated, SettingCopy);
+	}
+
+	return true;
+}
+
 void FConcertTakeRecorderManager::RegisterExtensions()
 {
 	ITakeRecorderModule& Module = FTakeRecorderRecorderManagerGetModule();
diff --git a/Engine/Plugins/VirtualProduction/MultiUserTakes/Source/ConcertTakeRecorder/Private/ConcertTakeRecorderManager.h b/Engine/Plugins/VirtualProduction/MultiUserTakes/Source/ConcertTakeRecorder/Private/ConcertTakeRecorderManager.h
--- a/Engine/Plugins/VirtualProduction/MultiUserTakes/Source/ConcertTakeRecorder/Private/ConcertTakeRecorderManager.h
+++ b/Engine/Plugins/VirtualProduction/MultiUserTakes/Source/ConcertTakeRecorder/Private/ConcertTakeRecorderManager.h
@@ -2,38 +2,41 @@
 
 #pragma once
 
-
+#include "Misc/Optional.h"
+#include "Templates/FunctionFwd.h"
 #include "UObject/GCObject.h"
 
 class FConcertTakeRecorderClientSessionCustomization;
-class SWidget;
+class IConcertSyncClient;
+class IConcertClientSession;
 class ULevelSequence;
 class UTakePreset;
+class UTakeRecorder;
+class SWidget;
+
+enum class ECheckBoxState : uint8;
 enum class EConcertClientStatus : uint8;
 enum class EConcertConnectionStatus : uint8;
 enum class ETransactionFilterResult : uint8;
 enum class EPackageFilterResult : uint8;
+
+struct EVisibility;
 struct FConcertClientRecordSetting;
 struct FConcertMultiUserSyncChangeEvent;
 struct FConcertRecordSettingsChangeEvent;
 struct FConcertRecordingCancelledEvent;
 struct FConcertRecordingFinishedEvent;
 struct FConcertRecordingNamedLevelSequenceEvent;
+struct FConcertSessionContext;
 struct FConcertSessionClientInfo;
 struct FConcertTakeInitializedEvent;
 struct FConcertPackageInfo;
 struct FFrameNumber;
 struct FTakeRecorderParameters;
+struct FTakeRecordSettings;
 
 DECLARE_LOG_CATEGORY_EXTERN(LogConcertTakeRecorder, Log, All);
 
-class IConcertSyncClient;
-class IConcertClientSession;
-struct FConcertSessionContext;
-class UTakeRecorder;
-enum class ECheckBoxState : uint8;
-struct EVisibility;
-
 /**
  * Take Recorder manager that is held by the client sync module that keeps track of when a take is started, stopped or cancelled.
  * Events are registered to client sessions that will then operate on the Take Recorder UIs
@@ -41,15 +44,11 @@ struct EVisibility;
 class FConcertTakeRecorderManager : public FGCObject
 {
 public:
-	/**
-	 * Constructor - registers TakeRecorderInitialized handler with the take recorder module
-	 */
+	
+	/** Registers TakeRecorderInitialized handler with the take recorder module */
 	FConcertTakeRecorderManager();
-
-	/**
-	 * Destructor - unregisters TakeRecorderInitialized handler from the take recorder module
-	 */
-	~FConcertTakeRecorderManager();
+	/** Unregisters TakeRecorderInitialized handler from the take recorder module */
+	virtual ~FConcertTakeRecorderManager() override;
 
 	/**
 	 * Register all custom take recorder events for the specified client session
@@ -57,7 +56,6 @@ public:
 	 * @param InSession	The client session to register custom events with
 	 */
 	void Register(TSharedRef<IConcertClientSession> InSession);
-
 	/**
 	 * Unregister previously registered custom take recorder events from the specified client session
 	 *
@@ -65,6 +63,18 @@ public:
 	 */
 	void Unregister(TSharedRef<IConcertClientSession> InSession);
 
+	/** Attempts the client recorder settings for the given client. The other clients inform us when they are changed so they are somewhat in sync. Nullptr if the client is unknown. */
+	const FConcertClientRecordSetting* FindClientRecorderSetting(const FGuid& EndpointId) const;
+	/**
+	 * Modifies the settings of a client, if the client exists. Subsequently updates all the relevant systems of the change.
+	 * @return Whether the client was found.
+	 */
+	bool EditClientSettings(
+		const FGuid& EndpointId,
+		TFunctionRef<void(FTakeRecordSettings& Settings)> ModifierFunc,
+		TOptional<TFunctionRef<bool(const FTakeRecordSettings& Settings)>> Predicate = {}
+		);
+
 private:
 	//~ Take recorder delegate handlers
 	void OnTakeRecorderInitialized(UTakeRecorder* TakeRecorder);
diff --git a/Engine/Plugins/VirtualProduction/MultiUserTakes/Source/ConcertTakeRecorder/Private/ConcertTakeRecorderModule.cpp b/Engine/Plugins/VirtualProduction/MultiUserTakes/Source/ConcertTakeRecorder/Private/ConcertTakeRecorderModule.cpp
--- a/Engine/Plugins/VirtualProduction/MultiUserTakes/Source/ConcertTakeRecorder/Private/ConcertTakeRecorderModule.cpp
+++ b/Engine/Plugins/VirtualProduction/MultiUserTakes/Source/ConcertTakeRecorder/Private/ConcertTakeRecorderModule.cpp
@@ -1,35 +1,27 @@
 // Copyright Epic Games, Inc. All Rights Reserved.
 
-#include "Modules/ModuleManager.h"
+#include "ConcertTakeRecorderModule.h"
+
 #include "ConcertTakeRecorderManager.h"
 #include "ConcertTakeRecorderStyle.h"
 
 LLM_DEFINE_TAG(Concert_ConcertTakeRecorder);
 
-/**
- * Module that adds multi user synchronization to take recorder.
- */
-class FConcertTakeRecorderModule : public IModuleInterface
+namespace UE::ConcertTakeRecorder
 {
-private:
-	virtual void StartupModule() override
+	void FConcertTakeRecorderModule::StartupModule()
 	{
 		LLM_SCOPE_BYTAG(Concert_ConcertTakeRecorder);
+		
 		FConcertTakeRecorderStyle::Initialize();
 		ConcertManager = MakeUnique<FConcertTakeRecorderManager>();
 	}
-
-	virtual void ShutdownModule() override
+	
+	void FConcertTakeRecorderModule::ShutdownModule()
 	{
 		ConcertManager.Reset();
 		FConcertTakeRecorderStyle::Shutdown();
 	}
+}
 
-	virtual ~FConcertTakeRecorderModule() = default;
-
-
-	TUniquePtr<FConcertTakeRecorderManager> ConcertManager;
-};
-
-
-IMPLEMENT_MODULE(FConcertTakeRecorderModule, ConcertTakeRecorder);
+IMPLEMENT_MODULE(UE::ConcertTakeRecorder::FConcertTakeRecorderModule, ConcertTakeRecorder);
diff --git a/Engine/Plugins/VirtualProduction/MultiUserTakes/Source/ConcertTakeRecorder/Private/ConcertTakeRecorderModule.h b/Engine/Plugins/VirtualProduction/MultiUserTakes/Source/ConcertTakeRecorder/Private/ConcertTakeRecorderModule.h
new file mode 100644
--- /dev/null
+++ b/Engine/Plugins/VirtualProduction/MultiUserTakes/Source/ConcertTakeRecorder/Private/ConcertTakeRecorderModule.h
@@ -0,0 +1,40 @@
+// Copyright Epic Games, Inc. All Rights Reserved.
+
+#pragma once
+
+#include "ConcertTakeRecorderManager.h"
+#include "Modules/ModuleManager.h"
+#include "Templates/UniquePtr.h"
+
+namespace UE::ConcertTakeRecorder
+{
+	/** Module that adds multi user synchronization to take recorder. */
+	class FConcertTakeRecorderModule : public IModuleInterface
+	{
+	public:
+
+		/**
+		 * Singleton-like access to this module's interface.  This is just for convenience!
+		 * Beware of calling this during the shutdown phase, though.  Your module might have been unloaded already.
+		 *
+		 * @return Returns singleton instance, loading the module on demand if needed
+		 */
+		static inline FConcertTakeRecorderModule& Get()
+		{
+			static const FName ModuleName = "ConcertTakeRecorder";
+			return FModuleManager::LoadModuleChecked<FConcertTakeRecorderModule>(ModuleName);
+		}
+		
+		//~ Begin IModuleInterface Interface
+		virtual void StartupModule() override;
+		virtual void ShutdownModule() override;
+		//~ End IModuleInterface Interface
+
+		FConcertTakeRecorderManager* GetTakeRecorderManager() const { return ConcertManager.Get(); }
+		
+	private:
+
+		/** Synchronizes takes across clients. Valid for the lifetime of the module. */
+		TUniquePtr<FConcertTakeRecorderManager> ConcertManager;
+	};
+}
diff --git a/Engine/Plugins/VirtualProduction/MultiUserTakes/Source/ConcertTakeRecorder/Private/MultiUserTakesFunctionLibrary.cpp b/Engine/Plugins/VirtualProduction/MultiUserTakes/Source/ConcertTakeRecorder/Private/MultiUserTakesFunctionLibrary.cpp
new file mode 100644
--- /dev/null
+++ b/Engine/Plugins/VirtualProduction/MultiUserTakes/Source/ConcertTakeRecorder/Private/MultiUserTakesFunctionLibrary.cpp
@@ -0,0 +1,94 @@
+// Copyright Epic Games, Inc. All Rights Reserved.
+
+#include "MultiUserTakesFunctionLibrary.h"
+
+#include "ConcertLogGlobal.h"
+#include "ConcertTakeRecorderManager.h"
+#include "ConcertTakeRecorderMessages.h"
+#include "ConcertTakeRecorderModule.h"
+#include "ConcertTakeRecorderSynchronizationCustomization.h"
+#include "IConcertSyncClient.h"
+#include "IConcertSyncClientModule.h"
+
+#include "Templates/Function.h"
+
+bool UMultiUserTakesFunctionLibrary::GetRecordOnClientLocal()
+{
+	const TSharedPtr<IConcertSyncClient> ConcertSyncClient = IConcertSyncClientModule::Get().GetClient(TEXT("MultiUser"));
+	const TSharedPtr<IConcertClientSession> Session = ConcertSyncClient ? ConcertSyncClient->GetConcertClient()->GetCurrentSession() : nullptr;
+	
+	UE_CLOG(!Session, LogConcert, Warning, TEXT("MultiUserTakes: Cannot change local RecordOnClient state because this editor is nott in any session!"));
+	return Session && GetRecordOnClient(Session->GetSessionClientEndpointId());
+}
+
+void UMultiUserTakesFunctionLibrary::SetRecordOnClientLocal(bool bNewValue)
+{
+	const TSharedPtr<IConcertSyncClient> ConcertSyncClient = IConcertSyncClientModule::Get().GetClient(TEXT("MultiUser"));
+	const TSharedPtr<IConcertClientSession> Session = ConcertSyncClient ? ConcertSyncClient->GetConcertClient()->GetCurrentSession() : nullptr;
+	if (Session)
+	{
+		SetRecordOnClient(Session->GetSessionClientEndpointId(), bNewValue);
+	}
+	else
+	{
+		UE_LOG(LogConcert, Warning, TEXT("MultiUserTakes: Cannot change local RecordOnClient state because this editor is nott in any session!"));
+	}
+}
+
+bool UMultiUserTakesFunctionLibrary::GetRecordOnClient(const FGuid& ClientEndpointId)
+{
+	using namespace UE::ConcertTakeRecorder;
+	
+	if (const FConcertTakeRecorderManager* Manager = FConcertTakeRecorderModule::Get().GetTakeRecorderManager())
+	{
+		const FConcertClientRecordSetting* Settings = Manager->FindClientRecorderSetting(ClientEndpointId);
+		UE_CLOG(!Settings, LogConcert, Warning, TEXT("MultiUserTakes: Unknown client ID %s"), *ClientEndpointId.ToString());
+		return Settings && Settings->Settings.bRecordOnClient;
+	}
+	return false;
+}
+
+void UMultiUserTakesFunctionLibrary::SetRecordOnClient(const FGuid& ClientEndpointId, bool bNewValue)
+{
+	using namespace UE::ConcertTakeRecorder;
+	
+	if (FConcertTakeRecorderManager* Manager = FConcertTakeRecorderModule::Get().GetTakeRecorderManager())
+	{
+		const bool bClientWasFound = Manager->EditClientSettings(
+			ClientEndpointId,
+			[bNewValue](FTakeRecordSettings& Settings) { Settings.bRecordOnClient = bNewValue; },
+			{ [bNewValue](const FTakeRecordSettings& Settings){ return Settings.bRecordOnClient != bNewValue; } }
+			);
+		UE_CLOG(!bClientWasFound, LogConcert, Warning, TEXT("MultiUserTakes: Unknown client ID %s"), *ClientEndpointId.ToString());
+	}
+}
+
+bool UMultiUserTakesFunctionLibrary::GetSynchronizeTakeRecorderTransactionsLocal()
+{
+	return GetMutableDefault<UConcertTakeSynchronization>()->bSyncTakeRecordingTransactions;
+}
+
+bool UMultiUserTakesFunctionLibrary::GetSynchronizeTakeRecorderTransactions(const FGuid& ClientEndpointId)
+{
+	using namespace UE::ConcertTakeRecorder;
+	
+	if (const FConcertTakeRecorderManager* Manager = FConcertTakeRecorderModule::Get().GetTakeRecorderManager())
+	{
+		const FConcertClientRecordSetting* Settings = Manager->FindClientRecorderSetting(ClientEndpointId);
+		UE_CLOG(!Settings, LogConcert, Warning, TEXT("MultiUserTakes: Unknown client ID %s"), *ClientEndpointId.ToString());
+		return Settings && Settings->bTakeSyncEnabled;
+	}
+	return false;
+}
+
+void UMultiUserTakesFunctionLibrary::SetSynchronizeTakeRecorderTransactionsLocal(bool bNewValue)
+{
+	UConcertTakeSynchronization* Settings = GetMutableDefault<UConcertTakeSynchronization>();
+	if (Settings->bSyncTakeRecordingTransactions != bNewValue)
+	{
+		Settings->bSyncTakeRecordingTransactions = bNewValue;
+
+		// Causes all the other clients to be notified. This should REALLY live in UConcertTakeSynchronization.
+		FConcertTakeRecorderSynchronizationCustomization::OnSyncPropertyValueChanged().Broadcast(bNewValue);
+	}
+}
diff --git a/Engine/Plugins/VirtualProduction/MultiUserTakes/Source/ConcertTakeRecorder/Public/MultiUserTakesFunctionLibrary.h b/Engine/Plugins/VirtualProduction/MultiUserTakes/Source/ConcertTakeRecorder/Public/MultiUserTakesFunctionLibrary.h
new file mode 100644
--- /dev/null
+++ b/Engine/Plugins/VirtualProduction/MultiUserTakes/Source/ConcertTakeRecorder/Public/MultiUserTakesFunctionLibrary.h
@@ -0,0 +1,70 @@
+// Copyright Epic Games, Inc. All Rights Reserved.
+
+#pragma once
+
+#include "Kismet/BlueprintFunctionLibrary.h"
+#include "MultiUserTakesFunctionLibrary.generated.h"
+
+/** Exposes some of the recorder settings to Blueprints to allow for more editor scripting capabilities. */
+UCLASS()
+class CONCERTTAKERECORDER_API UMultiUserTakesFunctionLibrary : public UBlueprintFunctionLibrary
+{
+	GENERATED_BODY()
+public:
+
+	/**
+	 * Gets the checkbox value in the "Record On Client" column in the settings displayed at the bottom of the "Take Recorder" tab.
+	 * This is an utility for getting the local client's value; GetRecordOnClient can also be used.
+	 */
+	UFUNCTION(BlueprintPure, Category = "Multi User Takes")
+	static bool GetRecordOnClientLocal();
+
+	/**
+	 * Sets the checkbox value in the "Record On Client" column in the settings displayed at the bottom of the "Take Recorder" tab.
+	 * This is an utility for setting the local client's value; SetRecordOnClient can also be used.
+	 */
+	UFUNCTION(BlueprintCallable, Category = "Multi User Takes")
+	static void SetRecordOnClientLocal(bool bNewValue);
+
+	/** 
+	 * Gets the checkbox value in the "Record On Client" column in the settings displayed at the bottom of the "Take Recorder" tab.
+	 * This function queries the state of any connected client. For the local client, you can also use GetRecordOnClientLocal.
+	 *
+	 * @param ClientEndpointId ID of the client. You can get this by using e.g.UMultiUserSubsystem::GetRemoteClientIds or UMultiUserSubsystem::GetLocalClientId.
+	 */
+	UFUNCTION(BlueprintPure, Category = "Multi User Takes")
+	static bool GetRecordOnClient(const FGuid& ClientEndpointId);
+	
+	/** 
+	 * Sets the checkbox value in the "Record On Client" column in the settings displayed at the bottom of the "Take Recorder" tab.
+	 * This function queries the state of any connected client. For the local client, you can also use SetRecordOnClientLocal.
+	 *
+	 * @param ClientEndpointId ID of the client. You can get this by using e.g.UMultiUserSubsystem::GetRemoteClientIds or UMultiUserSubsystem::GetLocalClientId.
+	 * @param bNewValue The new value to set
+	 */
+	UFUNCTION(BlueprintCallable, Category = "Multi User Takes")
+	static void SetRecordOnClient(const FGuid& ClientEndpointId, bool bNewValue);
+
+	/**
+	 * Gets the value of the "SynchronizeTakeRecorderTransactions" checkbox in the settings displayed at the bottom of the "Take Recorder" tab.
+	 * This is an utility for getting the local client's value; GetSynchronizeTakeRecorderTransactions can also be used.
+	 */
+	UFUNCTION(BlueprintPure, Category = "Multi User Takes")
+	static bool GetSynchronizeTakeRecorderTransactionsLocal();
+
+	/**
+	 * Gets the value of the "SynchronizeTakeRecorderTransactions" checkbox in the settings displayed at the bottom of the "Take Recorder" tab.
+	 * This function queries the state of any connected client. For the local client, you can also use GetSynchronizeTakeRecorderTransactionsLocal.
+	 *
+	 * @param ClientEndpointId ID of the client. You can get this by using e.g.UMultiUserSubsystem::GetRemoteClientIds or UMultiUserSubsystem::GetLocalClientId.
+	 */
+	UFUNCTION(BlueprintPure, Category = "Multi User Takes")
+	static bool GetSynchronizeTakeRecorderTransactions(const FGuid& ClientEndpointId);
+
+	/**
+	 * Sets the value of the "SynchronizeTakeRecorderTransactions" checkbox in the settings displayed at the bottom of the "Take Recorder" tab.
+	 * @note Only the value of the local client can be set. Setting of remote clients is not implemented (no technical reason - there just never was a use-case). 
+	 */
+	UFUNCTION(BlueprintCallable, Category = "Multi User Takes")
+	static void SetSynchronizeTakeRecorderTransactionsLocal(bool bNewValue);
+};
diff --git a/Engine/Plugins/VirtualProduction/VirtualCamera/Source/VirtualCamera/Private/FunctionLibraries/MultiUserTakesVCamFunctionLibrary.cpp b/Engine/Plugins/VirtualProduction/VirtualCamera/Source/VirtualCamera/Private/FunctionLibraries/MultiUserTakesVCamFunctionLibrary.cpp
new file mode 100644
--- /dev/null
+++ b/Engine/Plugins/VirtualProduction/VirtualCamera/Source/VirtualCamera/Private/FunctionLibraries/MultiUserTakesVCamFunctionLibrary.cpp
@@ -0,0 +1,64 @@
+// Copyright Epic Games, Inc. All Rights Reserved.
+
+#include "FunctionLibraries/MultiUserTakesVCamFunctionLibrary.h"
+
+#if WITH_EDITOR
+#include "MultiUserTakesFunctionLibrary.h"
+#endif
+
+bool UMultiUserTakesVCamFunctionLibrary::GetRecordOnClientLocal()
+{
+#if WITH_EDITOR
+	return UMultiUserTakesFunctionLibrary::GetRecordOnClientLocal();
+#else
+	return false;
+#endif
+}
+
+void UMultiUserTakesVCamFunctionLibrary::SetRecordOnClientLocal(bool bNewValue)
+{
+#if WITH_EDITOR
+	UMultiUserTakesFunctionLibrary::SetRecordOnClientLocal(bNewValue);
+#endif
+}
+
+bool UMultiUserTakesVCamFunctionLibrary::GetRecordOnClient(const FGuid& ClientEndpointId)
+{
+#if WITH_EDITOR
+	return UMultiUserTakesFunctionLibrary::GetRecordOnClient(ClientEndpointId);
+#else
+	return false;
+#endif
+}
+
+void UMultiUserTakesVCamFunctionLibrary::SetRecordOnClient(const FGuid& ClientEndpointId, bool bNewValue)
+{
+#if WITH_EDITOR
+	UMultiUserTakesFunctionLibrary::SetRecordOnClient(ClientEndpointId, bNewValue);
+#endif
+}
+
+bool UMultiUserTakesVCamFunctionLibrary::GetSynchronizeTakeRecorderTransactionsLocal()
+{
+#if WITH_EDITOR
+	return UMultiUserTakesFunctionLibrary::GetSynchronizeTakeRecorderTransactionsLocal();
+#else
+	return false;
+#endif
+}
+
+bool UMultiUserTakesVCamFunctionLibrary::GetSynchronizeTakeRecorderTransactions(const FGuid& ClientEndpointId)
+{
+#if WITH_EDITOR
+	return UMultiUserTakesFunctionLibrary::GetSynchronizeTakeRecorderTransactions(ClientEndpointId);
+#else
+	return false;
+#endif
+}
+
+void UMultiUserTakesVCamFunctionLibrary::SetSynchronizeTakeRecorderTransactionsLocal(bool bNewValue)
+{
+#if WITH_EDITOR
+	UMultiUserTakesFunctionLibrary::SetSynchronizeTakeRecorderTransactionsLocal(bNewValue);
+#endif
+}
diff --git a/Engine/Plugins/VirtualProduction/VirtualCamera/Source/VirtualCamera/Public/FunctionLibraries/MultiUserTakesVCamFunctionLibrary.h b/Engine/Plugins/VirtualProduction/VirtualCamera/Source/VirtualCamera/Public/FunctionLibraries/MultiUserTakesVCamFunctionLibrary.h
new file mode 100644
--- /dev/null
+++ b/Engine/Plugins/VirtualProduction/VirtualCamera/Source/VirtualCamera/Public/FunctionLibraries/MultiUserTakesVCamFunctionLibrary.h
@@ -0,0 +1,73 @@
+// Copyright Epic Games, Inc. All Rights Reserved.
+
+#pragma once
+
+#include "Kismet/BlueprintFunctionLibrary.h"
+#include "MultiUserTakesVCamFunctionLibrary.generated.h"
+
+/**
+ * This library wraps UMultiUserTakesFunctionLibrary: in editor builds, the calls are forwarded and cooked builds, the calls are compiled out.
+ * This is done so that Blueprints can continue to be cooked without crashing due to missing functions.
+ */
+UCLASS()
+class VIRTUALCAMERA_API UMultiUserTakesVCamFunctionLibrary : public UBlueprintFunctionLibrary
+{
+	GENERATED_BODY()
+public:
+
+	/**
+	 * Gets the checkbox value in the "Record On Client" column in the settings displayed at the bottom of the "Take Recorder" tab.
+	 * This is an utility for getting the local client's value; GetRecordOnClient can also be used.
+	 */
+	UFUNCTION(BlueprintPure, Category = "VirtualCamera | Multi User Takes")
+	static bool GetRecordOnClientLocal();
+
+	/**
+	 * Sets the checkbox value in the "Record On Client" column in the settings displayed at the bottom of the "Take Recorder" tab.
+	 * This is an utility for setting the local client's value; SetRecordOnClient can also be used.
+	 */
+	UFUNCTION(BlueprintCallable, Category = "VirtualCamera | Multi User Takes")
+	static void SetRecordOnClientLocal(bool bNewValue);
+
+	/** 
+	 * Gets the checkbox value in the "Record On Client" column in the settings displayed at the bottom of the "Take Recorder" tab.
+	 * This function queries the state of any connected client. For the local client, you can also use GetRecordOnClientLocal.
+	 *
+	 * @param ClientEndpointId ID of the client. You can get this by using e.g.UMultiUserSubsystem::GetRemoteClientIds or UMultiUserSubsystem::GetLocalClientId.
+	 */
+	UFUNCTION(BlueprintPure, Category = "VirtualCamera | Multi User Takes")
+	static bool GetRecordOnClient(const FGuid& ClientEndpointId);
+	
+	/** 
+	 * Sets the checkbox value in the "Record On Client" column in the settings displayed at the bottom of the "Take Recorder" tab.
+	 * This function queries the state of any connected client. For the local client, you can also use SetRecordOnClientLocal.
+	 *
+	 * @param ClientEndpointId ID of the client. You can get this by using e.g.UMultiUserSubsystem::GetRemoteClientIds or UMultiUserSubsystem::GetLocalClientId.
+	 * @param bNewValue The new value to set
+	 */
+	UFUNCTION(BlueprintCallable, Category = "VirtualCamera | Multi User Takes")
+	static void SetRecordOnClient(const FGuid& ClientEndpointId, bool bNewValue);
+
+	/**
+	 * Gets the value of the "SynchronizeTakeRecorderTransactions" checkbox in the settings displayed at the bottom of the "Take Recorder" tab.
+	 * This is an utility for getting the local client's value; GetSynchronizeTakeRecorderTransactions can also be used.
+	 */
+	UFUNCTION(BlueprintPure, Category = "VirtualCamera | Multi User Takes")
+	static bool GetSynchronizeTakeRecorderTransactionsLocal();
+
+	/**
+	 * Gets the value of the "SynchronizeTakeRecorderTransactions" checkbox in the settings displayed at the bottom of the "Take Recorder" tab.
+	 * This function queries the state of any connected client. For the local client, you can also use GetSynchronizeTakeRecorderTransactionsLocal.
+	 *
+	 * @param ClientEndpointId ID of the client. You can get this by using e.g.UMultiUserSubsystem::GetRemoteClientIds or UMultiUserSubsystem::GetLocalClientId.
+	 */
+	UFUNCTION(BlueprintPure, Category = "VirtualCamera | Multi User Takes")
+	static bool GetSynchronizeTakeRecorderTransactions(const FGuid& ClientEndpointId);
+
+	/**
+	 * Sets the value of the "SynchronizeTakeRecorderTransactions" checkbox in the settings displayed at the bottom of the "Take Recorder" tab.
+	 * @note Only the value of the local client can be set. Setting of remote clients is not implemented (no technical reason - there just never was a use-case). 
+	 */
+	UFUNCTION(BlueprintCallable, Category = "VirtualCamera | Multi User Takes")
+	static void SetSynchronizeTakeRecorderTransactionsLocal(bool bNewValue);
+};
diff --git a/Engine/Plugins/VirtualProduction/VirtualCamera/Source/VirtualCamera/VirtualCamera.Build.cs b/Engine/Plugins/VirtualProduction/VirtualCamera/Source/VirtualCamera/VirtualCamera.Build.cs
--- a/Engine/Plugins/VirtualProduction/VirtualCamera/Source/VirtualCamera/VirtualCamera.Build.cs
+++ b/Engine/Plugins/VirtualProduction/VirtualCamera/Source/VirtualCamera/VirtualCamera.Build.cs
@@ -40,16 +40,24 @@ public class VirtualCamera : ModuleRules
 		
 		if (Target.bBuildEditor == true)
 		{
-			PublicDependencyModuleNames.Add("LevelSequenceEditor");
-			PublicDependencyModuleNames.Add("Sequencer");
-			PublicDependencyModuleNames.Add("SlateCore");
-			PublicDependencyModuleNames.Add("TakeRecorder");
-			PrivateDependencyModuleNames.Add("LevelEditor");
-			PrivateDependencyModuleNames.Add("EditorFramework");
-			PrivateDependencyModuleNames.Add("UnrealEd");
-			PrivateDependencyModuleNames.Add("EditorScriptingUtilities");
-			PrivateDependencyModuleNames.Add("VPUtilitiesEditor");
-			PrivateDependencyModuleNames.Add("TakesCore");
+			PublicDependencyModuleNames.AddRange(new string[]
+			{
+				"LevelSequenceEditor",
+				"Sequencer",
+				"SlateCore",
+				"TakeRecorder"
+			});
+			
+			PrivateDependencyModuleNames.AddRange(new string[]
+			{
+				"ConcertTakeRecorder",
+				"EditorFramework",
+				"EditorScriptingUtilities",
+				"LevelEditor",
+				"TakesCore",
+				"UnrealEd",
+				"VPUtilitiesEditor",
+			});
 		}
 	}
 }
diff --git a/Engine/Plugins/VirtualProduction/VirtualCamera/VirtualCamera.uplugin b/Engine/Plugins/VirtualProduction/VirtualCamera/VirtualCamera.uplugin
--- a/Engine/Plugins/VirtualProduction/VirtualCamera/VirtualCamera.uplugin
+++ b/Engine/Plugins/VirtualProduction/VirtualCamera/VirtualCamera.uplugin
@@ -38,11 +38,15 @@
 	],
 	"Plugins": [
 		{
-			"Name" : "VirtualCameraCore",
+			"Name" : "CineCameraRigs",
 			"Enabled": true
 		},
 		{
-			"Name" : "CineCameraRigs",
+			"Name" : "MultiUserTakes",
+			"Enabled": true
+		},
+		{
+			"Name" : "VirtualCameraCore",
 			"Enabled": true
 		}
 	]
