Make sure _temp takes are properly filtered out when in a multi-user take record.

diff --git a/Engine/Plugins/VirtualProduction/MultiUserTakes/Source/ConcertTakeRecorder/Private/ConcertTakeRecorderManager.cpp b/Engine/Plugins/VirtualProduction/MultiUserTakes/Source/ConcertTakeRecorder/Private/ConcertTakeRecorderManager.cpp
--- a/Engine/Plugins/VirtualProduction/MultiUserTakes/Source/ConcertTakeRecorder/Private/ConcertTakeRecorderManager.cpp
+++ b/Engine/Plugins/VirtualProduction/MultiUserTakes/Source/ConcertTakeRecorder/Private/ConcertTakeRecorderManager.cpp
@@ -960,7 +960,7 @@ EPackageFilterResult FConcertTakeRecorderManager::ShouldPackageBeFiltered(const
 	{
 		FTakeRecorderProjectParameters Project = GetDefault<UTakeRecorderProjectSettings>()->Settings;
 		FString FullName = InPackageInfo.PackageName.ToString();
-		if (FullName.StartsWith(Project.RootTakeSaveDir.Path) &&
+		if (FullName.Contains(Project.RootTakeSaveDir.Path) &&
 			FullName.EndsWith(UE::TakeRecorderManager::Private::TempPostfix))
 		{
 			return EPackageFilterResult::Exclude;
