Change default state for Transact Meta Data property when in a Multi-user take recorder session.

diff --git a/Engine/Plugins/VirtualProduction/MultiUserTakes/Source/ConcertTakeRecorder/Private/ConcertTakeRecorderMessages.h b/Engine/Plugins/VirtualProduction/MultiUserTakes/Source/ConcertTakeRecorder/Private/ConcertTakeRecorderMessages.h
--- a/Engine/Plugins/VirtualProduction/MultiUserTakes/Source/ConcertTakeRecorder/Private/ConcertTakeRecorderMessages.h
+++ b/Engine/Plugins/VirtualProduction/MultiUserTakes/Source/ConcertTakeRecorder/Private/ConcertTakeRecorderMessages.h
@@ -18,7 +18,7 @@ public:
 	bool bSyncTakeRecordingTransactions = true;
 
 	UPROPERTY(config,EditAnywhere,BlueprintReadWrite,Category="Multi-user Take Synchronization")
-	bool bTransactTakeMetadata = false;
+	bool bTransactTakeMetadata = true;
 };
 
 USTRUCT()
