Ensure we don't save takes when we are not recording in take recorder in a Multi-user session.

diff --git a/Engine/Plugins/VirtualProduction/MultiUserTakes/Source/ConcertTakeRecorder/Private/ConcertTakeRecorderManager.cpp b/Engine/Plugins/VirtualProduction/MultiUserTakes/Source/ConcertTakeRecorder/Private/ConcertTakeRecorderManager.cpp
--- a/Engine/Plugins/VirtualProduction/MultiUserTakes/Source/ConcertTakeRecorder/Private/ConcertTakeRecorderManager.cpp
+++ b/Engine/Plugins/VirtualProduction/MultiUserTakes/Source/ConcertTakeRecorder/Private/ConcertTakeRecorderManager.cpp
@@ -961,7 +961,7 @@ EPackageFilterResult FConcertTakeRecorderManager::ShouldPackageBeFiltered(const
 		FTakeRecorderProjectParameters Project = GetDefault<UTakeRecorderProjectSettings>()->Settings;
 		FString FullName = InPackageInfo.PackageName.ToString();
 		if (FullName.Contains(Project.RootTakeSaveDir.Path) &&
-			FullName.EndsWith(UE::TakeRecorderManager::Private::TempPostfix))
+			FullName.Contains(UE::TakeRecorderManager::Private::TempPostfix))
 		{
 			return EPackageFilterResult::Exclude;
 		}
@@ -1032,6 +1032,7 @@ FTakeRecorderParameters FConcertTakeRecorderManager::SetupTakeParametersForMulti
 		{
 			FTakeRecorderParameters Output = Input;
 			Output.Project.TakeSaveDir = Input.Project.TakeSaveDir + UE::TakeRecorderManager::Private::TempPostfix;
+			Output.User.bSaveRecordedAssets = false;
 			return Output;
 		}
 	}
