Fix log spam due to Take Recorder tab not being opened. This happens because VCam Blueprint was calling the function every tick. Removed the log.

diff --git a/Engine/Plugins/VirtualProduction/MultiUserTakes/Source/ConcertTakeRecorder/Private/MultiUserTakesFunctionLibrary.cpp b/Engine/Plugins/VirtualProduction/MultiUserTakes/Source/ConcertTakeRecorder/Private/MultiUserTakesFunctionLibrary.cpp
--- a/Engine/Plugins/VirtualProduction/MultiUserTakes/Source/ConcertTakeRecorder/Private/MultiUserTakesFunctionLibrary.cpp
+++ b/Engine/Plugins/VirtualProduction/MultiUserTakes/Source/ConcertTakeRecorder/Private/MultiUserTakesFunctionLibrary.cpp
@@ -42,7 +42,6 @@ bool UMultiUserTakesFunctionLibrary::GetRecordOnClient(const FGuid& ClientEndpoi
 	if (const FConcertTakeRecorderManager* Manager = FConcertTakeRecorderModule::Get().GetTakeRecorderManager())
 	{
 		const FConcertClientRecordSetting* Settings = Manager->FindClientRecorderSetting(ClientEndpointId);
-		UE_CLOG(!Settings, LogConcert, Warning, TEXT("MultiUserTakes: Unknown client ID %s"), *ClientEndpointId.ToString());
 		return Settings && Settings->Settings.bRecordOnClient;
 	}
 	return false;
@@ -59,7 +58,6 @@ void UMultiUserTakesFunctionLibrary::SetRecordOnClient(const FGuid& ClientEndpoi
 			[bNewValue](FTakeRecordSettings& Settings) { Settings.bRecordOnClient = bNewValue; },
 			{ [bNewValue](const FTakeRecordSettings& Settings){ return Settings.bRecordOnClient != bNewValue; } }
 			);
-		UE_CLOG(!bClientWasFound, LogConcert, Warning, TEXT("MultiUserTakes: Unknown client ID %s"), *ClientEndpointId.ToString());
 	}
 }
 
@@ -75,7 +73,6 @@ bool UMultiUserTakesFunctionLibrary::GetSynchronizeTakeRecorderTransactions(cons
 	if (const FConcertTakeRecorderManager* Manager = FConcertTakeRecorderModule::Get().GetTakeRecorderManager())
 	{
 		const FConcertClientRecordSetting* Settings = Manager->FindClientRecorderSetting(ClientEndpointId);
-		UE_CLOG(!Settings, LogConcert, Warning, TEXT("MultiUserTakes: Unknown client ID %s"), *ClientEndpointId.ToString());
 		return Settings && Settings->bTakeSyncEnabled;
 	}
 	return false;
