Remove ensure from CanSkipHotReload in the take recorder.

diff --git a/Engine/Plugins/VirtualProduction/MultiUserTakes/Source/ConcertTakeRecorder/Private/ConcertTakeRecorderManager.cpp b/Engine/Plugins/VirtualProduction/MultiUserTakes/Source/ConcertTakeRecorder/Private/ConcertTakeRecorderManager.cpp
--- a/Engine/Plugins/VirtualProduction/MultiUserTakes/Source/ConcertTakeRecorder/Private/ConcertTakeRecorderManager.cpp
+++ b/Engine/Plugins/VirtualProduction/MultiUserTakes/Source/ConcertTakeRecorder/Private/ConcertTakeRecorderManager.cpp
@@ -973,7 +973,7 @@ EPackageFilterResult FConcertTakeRecorderManager::ShouldPackageBeFiltered(const
 
 bool FConcertTakeRecorderManager::CanSkipHotReload(const FConcertPackageInfo& InPackageInfo)
 {
-	if (IsTakeSyncEnabled() && ensure(WeakSession.IsValid()) && CanRecord())
+	if (IsTakeSyncEnabled() && WeakSession.IsValid() && CanRecord())
 	{
 		FTakeRecorderProjectParameters Project = GetDefault<UTakeRecorderProjectSettings>()->Settings;
 		FString FullName = InPackageInfo.PackageName.ToString();
