Change the default recording setting for Multi-user take records.  Now users have to opt-in to recording rather than a opt-out

diff --git a/Engine/Plugins/VirtualProduction/MultiUserTakes/Source/ConcertTakeRecorder/Private/ConcertTakeRecorderMessages.h b/Engine/Plugins/VirtualProduction/MultiUserTakes/Source/ConcertTakeRecorder/Private/ConcertTakeRecorderMessages.h
--- a/Engine/Plugins/VirtualProduction/MultiUserTakes/Source/ConcertTakeRecorder/Private/ConcertTakeRecorderMessages.h
+++ b/Engine/Plugins/VirtualProduction/MultiUserTakes/Source/ConcertTakeRecorder/Private/ConcertTakeRecorderMessages.h
@@ -27,7 +27,7 @@ struct FTakeRecordSettings
 	GENERATED_BODY();
 
 	UPROPERTY(config,EditAnywhere,Category="Multi-user Client Record Settings")
-	bool bRecordOnClient = true;
+	bool bRecordOnClient = false;
 
 	UPROPERTY(config,EditAnywhere,Category="Multi-user Client Record Settings")
 	bool bTransactSources = true;
