Add a CVar to control skip hot reload feature of take recorder.

diff --git a/Engine/Plugins/VirtualProduction/MultiUserTakes/Source/ConcertTakeRecorder/Private/ConcertTakeRecorderManager.cpp b/Engine/Plugins/VirtualProduction/MultiUserTakes/Source/ConcertTakeRecorder/Private/ConcertTakeRecorderManager.cpp
--- a/Engine/Plugins/VirtualProduction/MultiUserTakes/Source/ConcertTakeRecorder/Private/ConcertTakeRecorderManager.cpp
+++ b/Engine/Plugins/VirtualProduction/MultiUserTakes/Source/ConcertTakeRecorder/Private/ConcertTakeRecorderManager.cpp
@@ -50,6 +50,9 @@ static FAutoConsoleVariableRef  CVarEnableTakeSync(TEXT("Concert.EnableTakeRecor
 static int32 bConcertUseTakePresetPathForRecord = 0;
 static FAutoConsoleVariableRef  CVarEnableTakePresetPathSync(TEXT("Concert.UseTakePresetPath"), bConcertUseTakePresetPathForRecord, TEXT("Use Take Presets for Take Recording."));
 
+static int32 bEnableSkipHotReloadHint = 0;
+static FAutoConsoleVariableRef  CVarEnableHotReloadHint(TEXT("Concert.TakeRecorderSkipHotReloadHint"), bEnableSkipHotReloadHint, TEXT("Should we indicate to the clients that hot reload can be skipped for assets generated via take recorder."));
+
 #define LOCTEXT_NAMESPACE "ConcertTakeRecorder"
 
 ITakeRecorderModule& FTakeRecorderRecorderManagerGetModule()
@@ -973,7 +976,7 @@ EPackageFilterResult FConcertTakeRecorderManager::ShouldPackageBeFiltered(const
 
 bool FConcertTakeRecorderManager::CanSkipHotReload(const FConcertPackageInfo& InPackageInfo)
 {
-	if (IsTakeSyncEnabled() && WeakSession.IsValid() && CanRecord())
+	if (IsTakeSyncEnabled() && WeakSession.IsValid() && CanRecord() && bEnableSkipHotReloadHint > 0)
 	{
 		FTakeRecorderProjectParameters Project = GetDefault<UTakeRecorderProjectSettings>()->Settings;
 		FString FullName = InPackageInfo.PackageName.ToString();
