Show the record button even if no multi-user users have opt-in.  This means that the first person to click record will be the recorder.

diff --git a/Engine/Plugins/VirtualProduction/MultiUserTakes/Source/ConcertTakeRecorder/Private/ConcertTakeRecorderManager.cpp b/Engine/Plugins/VirtualProduction/MultiUserTakes/Source/ConcertTakeRecorder/Private/ConcertTakeRecorderManager.cpp
--- a/Engine/Plugins/VirtualProduction/MultiUserTakes/Source/ConcertTakeRecorder/Private/ConcertTakeRecorderManager.cpp
+++ b/Engine/Plugins/VirtualProduction/MultiUserTakes/Source/ConcertTakeRecorder/Private/ConcertTakeRecorderManager.cpp
@@ -16,6 +16,8 @@
 #include "ConcertSyncArchives.h"
 #include "ConcertTakeRecorderStyle.h"
 
+#include "MultiUserTakesFunctionLibrary.h"
+
 #include "ITakeRecorderModule.h"
 #include "Misc/Guid.h"
 #include "Recorder/TakeRecorder.h"
@@ -237,7 +239,6 @@ void FConcertTakeRecorderManager::RegisterExtensions()
 
 	Module.GetToolbarExtensionGenerators().AddRaw(this, &FConcertTakeRecorderManager::CreateExtensionWidget);
 	Module.GetRecordButtonExtensionGenerators().AddRaw(this, &FConcertTakeRecorderManager::CreateRecordButtonOverlay);
-	Module.GetRecordErrorCheckGenerator().AddRaw(this, &FConcertTakeRecorderManager::ReportRecordingError);
 	Module.GetCanReviewLastRecordedLevelSequenceDelegate().BindRaw(this, &FConcertTakeRecorderManager::CanReviewLastRecordedSequence);
 
 	if (GIsEditor)
@@ -328,8 +329,7 @@ bool FConcertTakeRecorderManager::ShouldIconBeVisible() const
 	UTakePreset* TakePreset = TakeRecorderModule.GetPendingTake();
 	UConcertTakeSynchronization const* TakeSync = GetDefault<UConcertTakeSynchronization>();
 
-	if (WeakSession.IsValid() && TakePreset && !bIsRecording && TakeSync->bSyncTakeRecordingTransactions
-		&& IsTakeSyncEnabled() && CanAnyRecord())
+	if (WeakSession.IsValid() && TakePreset && !bIsRecording && TakeSync->bSyncTakeRecordingTransactions && IsTakeSyncEnabled())
 	{
 		ULevelSequence* LevelSequence = TakePreset->GetLevelSequence();
 		if (LevelSequence)
@@ -432,6 +432,12 @@ void FConcertTakeRecorderManager::OnTakeRecorderInitialized(UTakeRecorder* TakeR
 	{
 		if (TSharedPtr<IConcertClientSession> Session = WeakSession.Pin())
 		{
+			if (!CanAnyRecord())
+			{
+				// If no-one can record then force the local client to have recording so that we can continue.
+				UMultiUserTakesFunctionLibrary::SetRecordOnClientLocal(true);
+			}
+
 			LastLevelSequence = nullptr;
 
 			ITakeRecorderModule& TakeRecorderModule = FModuleManager::LoadModuleChecked<ITakeRecorderModule>("TakeRecorder");
@@ -549,6 +555,7 @@ void FConcertTakeRecorderManager::OnTakeInitializedEvent(const FConcertSessionCo
 {
 	if (IsTakeSyncEnabled() && CanRecord())
 	{
+		TakeRecorderState.LastStoppedTake = "";
 		TakeRecorderState.LastStartedTake = InEvent.TakeName;
 
 		ITakeRecorderModule& TakeRecorderModule = FModuleManager::LoadModuleChecked<ITakeRecorderModule>("TakeRecorder");
@@ -799,18 +806,9 @@ void FConcertTakeRecorderManager::OnSessionConnectionChanged(IConcertClientSessi
 	}
 }
 
-void FConcertTakeRecorderManager::ReportRecordingError(FText &OutputError)
-{
-	if(WeakSession.IsValid() && IsTakeSyncEnabled() && !CanAnyRecord())
-	{
-		OutputError = LOCTEXT("ErrorWidget_NoRecorder", "No clients are available to record.");
-	}
-}
-
 bool FConcertTakeRecorderManager::CanAnyRecord() const
 {
 	UConcertSessionRecordSettings const* RecordSettings = GetDefault<UConcertSessionRecordSettings>();
-
 	bool bCanRecord = CanRecord();
 	for( const FConcertClientRecordSetting&  Remote : RecordSettings->RemoteSettings )
 	{
diff --git a/Engine/Plugins/VirtualProduction/MultiUserTakes/Source/ConcertTakeRecorder/Private/ConcertTakeRecorderManager.h b/Engine/Plugins/VirtualProduction/MultiUserTakes/Source/ConcertTakeRecorder/Private/ConcertTakeRecorderManager.h
--- a/Engine/Plugins/VirtualProduction/MultiUserTakes/Source/ConcertTakeRecorder/Private/ConcertTakeRecorderManager.h
+++ b/Engine/Plugins/VirtualProduction/MultiUserTakes/Source/ConcertTakeRecorder/Private/ConcertTakeRecorderManager.h
@@ -136,7 +136,6 @@ private:
 		return TEXT("FConcertTakeRecorderManager");
 	}
 
-	void ReportRecordingError(FText &);
 	bool CanRecord() const;
 	bool CanAnyRecord() const;
 
