Fix take recorder for recording on -game node.  There were several issues at play here:

1. Trying to open the sequencer from -game node resulted in the initialize recorded take to fail. This should not fail when editor is not present.
2. Reviewing the last take would fail because LoadObject would fail when the package hadn't fully arrived.
3. The sequencer manager would try to read the pending take to open it when in -game.
4. The ApplyOpenEvent would try to extract the pending take information for -game nodes and this code path is only needed for editor.

#rnx

diff --git a/Engine/Plugins/Developer/Concert/ConcertSync/ConcertSyncClient/Source/ConcertSyncClient/Private/ConcertClientSequencerManager.cpp b/Engine/Plugins/Developer/Concert/ConcertSync/ConcertSyncClient/Source/ConcertSyncClient/Private/ConcertClientSequencerManager.cpp
--- a/Engine/Plugins/Developer/Concert/ConcertSync/ConcertSyncClient/Source/ConcertSyncClient/Private/ConcertClientSequencerManager.cpp
+++ b/Engine/Plugins/Developer/Concert/ConcertSync/ConcertSyncClient/Source/ConcertSyncClient/Private/ConcertClientSequencerManager.cpp
@@ -905,7 +905,7 @@ void FConcertClientSequencerManager::ApplyOpenEvent(const FConcertSequencerOpenE
 		bDidOpen = true;
 	}
 
-	if (bDidOpen && UE::Private::ConcertClientSequencerManager::IsPendingTakePath(SequenceObjectPath))
+	if (GIsEditor && bDidOpen && UE::Private::ConcertClientSequencerManager::IsPendingTakePath(SequenceObjectPath))
 	{
 		ULevelSequence* PendingLevelSequence = FindObject<ULevelSequence>(nullptr, *UE::Private::ConcertClientSequencerManager::PendingTakePath);
 
@@ -932,7 +932,6 @@ void FConcertClientSequencerManager::ApplyOpenEvent(const FConcertSequencerOpenE
 		}
 
 		#if WITH_EDITOR
-		if (GIsEditor)
 		{
 			// We must specify the pending level sequence and not the pending take.  If we specify the pending take after we serialized into the existing
 			// asset the open asset will cause take recorder tab to assume that we want to edit the pending take where we really just want to view
diff --git a/Engine/Plugins/VirtualProduction/MultiUserTakes/Source/ConcertTakeRecorder/Private/ConcertTakeRecorderManager.cpp b/Engine/Plugins/VirtualProduction/MultiUserTakes/Source/ConcertTakeRecorder/Private/ConcertTakeRecorderManager.cpp
--- a/Engine/Plugins/VirtualProduction/MultiUserTakes/Source/ConcertTakeRecorder/Private/ConcertTakeRecorderManager.cpp
+++ b/Engine/Plugins/VirtualProduction/MultiUserTakes/Source/ConcertTakeRecorder/Private/ConcertTakeRecorderManager.cpp
@@ -31,6 +31,7 @@
 #include "LevelSequence.h"
 
 #include "PackageTools.h"
+#include "Misc/PackageName.h"
 
 #include "UObject/Package.h"
 #include "Widgets/Images/SImage.h"
@@ -592,14 +593,16 @@ void FConcertTakeRecorderManager::OnTakeInitializedEvent(const FConcertSessionCo
 
 			LastLevelSequence = nullptr;
 
-			UE_LOG(LogConcertTakeRecorder, Display, TEXT("Take initialized for recording by multi-user event."));
-			bIsRecording = true;
-			NewRecorder->Initialize(
-				LevelSequence,
-				Sources,
-				TakeMetadata,
-				DefaultParams
-			);
+			bIsRecording = NewRecorder->Initialize(LevelSequence, Sources, TakeMetadata, DefaultParams);
+			if (bIsRecording)
+			{
+				UE_LOG(LogConcertTakeRecorder, Display, TEXT("Take initialized for recording by multi-user event."));
+			}
+			else
+			{
+				UE_LOG(LogConcertTakeRecorder, Error, TEXT("Failed to initialize for recording by multi-user event."));
+			}
+
 		}
 	}
 }
@@ -1068,25 +1071,35 @@ FTakeRecorderParameters FConcertTakeRecorderManager::SetupTakeParametersForMulti
 
 void FConcertTakeRecorderManager::OnNamedLevelSequenceEvent(const FConcertSessionContext&, const FConcertRecordingNamedLevelSequenceEvent& InEvent)
 {
-	SetLastLevelSequence(LoadObject<ULevelSequence>(nullptr, *InEvent.LevelSequencePath));
+	LastLevelSequencePath = InEvent.LevelSequencePath;
 }
 
 bool FConcertTakeRecorderManager::CanReviewLastRecordedSequence() const
 {
 	if (WeakSession.IsValid())
 	{
+		// Due to when packages arrive in Multi-user, we may have received the name event message before the package is actually saved locally. To
+		// prevent loading issues, we need to check to see if that package has arrived and then load it to be able to review the last recording. This
+		// method is called by the panel to refresh the button used to review the last take but it requires a const call so LastLevelSequencePath and
+		// LastLevelSequence are mutable to allow changing in the const method.
+		if (!LastLevelSequencePath.IsEmpty())
+		{
+			FString PackageName;
+			FString AssetName;
+			LastLevelSequencePath.Split(TEXT("."), &PackageName, &AssetName);
+			if (FPackageName::DoesPackageExist(PackageName, nullptr))
+			{
+				LastLevelSequence = LoadObject<ULevelSequence>(nullptr, *LastLevelSequencePath);
+				
+				ITakeRecorderModule& Module = FTakeRecorderRecorderManagerGetModule();
+				Module.GetLastLevelSequenceProvider().ExecuteIfBound(LastLevelSequence);
+
+				LastLevelSequencePath.Reset();
+			}
+		}
 		return LastLevelSequence != nullptr;
 	}
 	return true;
 }
 
-
-void FConcertTakeRecorderManager::SetLastLevelSequence(ULevelSequence* InLastSequence)
-{
-	ITakeRecorderModule& Module = FTakeRecorderRecorderManagerGetModule();
-
-	LastLevelSequence = InLastSequence;
-	Module.GetLastLevelSequenceProvider().ExecuteIfBound(LastLevelSequence);
-}
-
 #undef LOCTEXT_NAMESPACE /*ConcertTakeRecorder*/
diff --git a/Engine/Plugins/VirtualProduction/MultiUserTakes/Source/ConcertTakeRecorder/Private/ConcertTakeRecorderManager.h b/Engine/Plugins/VirtualProduction/MultiUserTakes/Source/ConcertTakeRecorder/Private/ConcertTakeRecorderManager.h
--- a/Engine/Plugins/VirtualProduction/MultiUserTakes/Source/ConcertTakeRecorder/Private/ConcertTakeRecorderManager.h
+++ b/Engine/Plugins/VirtualProduction/MultiUserTakes/Source/ConcertTakeRecorder/Private/ConcertTakeRecorderManager.h
@@ -145,9 +145,8 @@ private:
 	void RecordSettingChange(const FConcertClientRecordSetting& RecordSetting);
 	void AddRemoteClient(const FConcertSessionClientInfo& ClientInfo);
 
-	void SetLastLevelSequence(ULevelSequence* InLastSequence);
 	bool CanReviewLastRecordedSequence() const;
-
+	
 	/** Weak pointer to the client session with which to send events. May be null or stale. */
 	TWeakPtr<IConcertClientSession> WeakSession;
 
@@ -161,8 +160,10 @@ private:
 	} TakeRecorderState;
 
 	TObjectPtr<UTakePreset> Preset = nullptr;
-	TObjectPtr<ULevelSequence> LastLevelSequence = nullptr;
 
+	mutable TObjectPtr<ULevelSequence> LastLevelSequence = nullptr;
+	mutable FString LastLevelSequencePath;
+	
 	TSharedPtr<FConcertTakeRecorderClientSessionCustomization> Customization;
 
 	/** Delegate for any changes in client state. */
diff --git a/Engine/Plugins/VirtualProduction/Takes/Source/TakeRecorder/Private/Recorder/TakeRecorder.cpp b/Engine/Plugins/VirtualProduction/Takes/Source/TakeRecorder/Private/Recorder/TakeRecorder.cpp
--- a/Engine/Plugins/VirtualProduction/Takes/Source/TakeRecorder/Private/Recorder/TakeRecorder.cpp
+++ b/Engine/Plugins/VirtualProduction/Takes/Source/TakeRecorder/Private/Recorder/TakeRecorder.cpp
@@ -658,7 +658,7 @@ bool UTakeRecorder::CreateDestinationAsset(const TCHAR* AssetPathFormat, ULevelS
 	FAssetRegistryModule::AssetCreated(SequenceAsset);
 
 	bool bSuccess = true;
-	if (bOpenSequencer)
+	if (bOpenSequencer && GIsEditor)
 	{
 		WeakSequencer = TakesUtils::OpenSequencer(SequenceAsset, OutError);
 		bSuccess = WeakSequencer.IsValid();
