Send project settings over the take initialized message so that other nodes have the right project settings to do the recording. This includes the root save directory.

#rnx

diff --git a/Engine/Plugins/VirtualProduction/MultiUserTakes/Source/ConcertTakeRecorder/Private/ConcertTakeRecorderManager.cpp b/Engine/Plugins/VirtualProduction/MultiUserTakes/Source/ConcertTakeRecorder/Private/ConcertTakeRecorderManager.cpp
--- a/Engine/Plugins/VirtualProduction/MultiUserTakes/Source/ConcertTakeRecorder/Private/ConcertTakeRecorderManager.cpp
+++ b/Engine/Plugins/VirtualProduction/MultiUserTakes/Source/ConcertTakeRecorder/Private/ConcertTakeRecorderManager.cpp
@@ -473,6 +473,7 @@ void FConcertTakeRecorderManager::OnTakeRecorderInitialized(UTakeRecorder* TakeR
 				TakeInitializedEvent.TakeName = TakeRecorder->GetName();
 				TakeInitializedEvent.TakePresetPath = TakeMetaData->GetPresetOrigin()->GetPathName();
 				TakeInitializedEvent.Settings = GetDefault<UTakeRecorderUserSettings>()->Settings;
+				TakeInitializedEvent.Project = GetDefault<UTakeRecorderProjectSettings>()->Settings;
 				TakeInitializedEvent.TakeMetaDataPath = TakeMetaData->GetPathName();
 
 				FConcertLocalIdentifierTable InLocalIdentifierTable;
@@ -587,7 +588,7 @@ void FConcertTakeRecorderManager::OnTakeInitializedEvent(const FConcertSessionCo
 
 			FTakeRecorderParameters DefaultParams;
 			DefaultParams.User = InEvent.Settings;
-			DefaultParams.Project = GetDefault<UTakeRecorderProjectSettings>()->Settings;
+			DefaultParams.Project = InEvent.Project;
 
 			UTakeRecorder* NewRecorder = NewObject<UTakeRecorder>(GetTransientPackage(), NAME_None, RF_Transient);
 
diff --git a/Engine/Plugins/VirtualProduction/MultiUserTakes/Source/ConcertTakeRecorder/Private/ConcertTakeRecorderMessages.h b/Engine/Plugins/VirtualProduction/MultiUserTakes/Source/ConcertTakeRecorder/Private/ConcertTakeRecorderMessages.h
--- a/Engine/Plugins/VirtualProduction/MultiUserTakes/Source/ConcertTakeRecorder/Private/ConcertTakeRecorderMessages.h
+++ b/Engine/Plugins/VirtualProduction/MultiUserTakes/Source/ConcertTakeRecorder/Private/ConcertTakeRecorderMessages.h
@@ -108,6 +108,9 @@ struct FConcertTakeInitializedEvent
 
 	UPROPERTY()
 	FTakeRecorderUserParameters Settings;
+
+	UPROPERTY()
+	FTakeRecorderProjectParameters Project;
 };
 
 USTRUCT()
