Warn users about having remote open enabled when recording on multiple clients.

[RN] When recording multiple clients in take recorder, it is recommended that you disable remote sequencer open. A error is now displayed to enforce this condition.

diff --git a/Engine/Plugins/VirtualProduction/MultiUserTakes/Source/ConcertTakeRecorder/Private/ConcertTakeRecorderManager.cpp b/Engine/Plugins/VirtualProduction/MultiUserTakes/Source/ConcertTakeRecorder/Private/ConcertTakeRecorderManager.cpp
--- a/Engine/Plugins/VirtualProduction/MultiUserTakes/Source/ConcertTakeRecorder/Private/ConcertTakeRecorderManager.cpp
+++ b/Engine/Plugins/VirtualProduction/MultiUserTakes/Source/ConcertTakeRecorder/Private/ConcertTakeRecorderManager.cpp
@@ -240,6 +240,7 @@ void FConcertTakeRecorderManager::RegisterExtensions()
 
 	Module.GetToolbarExtensionGenerators().AddRaw(this, &FConcertTakeRecorderManager::CreateExtensionWidget);
 	Module.GetRecordButtonExtensionGenerators().AddRaw(this, &FConcertTakeRecorderManager::CreateRecordButtonOverlay);
+	Module.GetRecordErrorCheckGenerator().AddRaw(this, &FConcertTakeRecorderManager::ReportRecordingError);
 	Module.GetCanReviewLastRecordedLevelSequenceDelegate().BindRaw(this, &FConcertTakeRecorderManager::CanReviewLastRecordedSequence);
 
 	if (GIsEditor)
@@ -332,11 +333,17 @@ bool FConcertTakeRecorderManager::ShouldIconBeVisible() const
 
 	if (WeakSession.IsValid() && TakePreset && !bIsRecording && TakeSync->bSyncTakeRecordingTransactions && IsTakeSyncEnabled())
 	{
+		bool bShouldSupressBecauseOfRemoteOpen = false;
+		if (TSharedPtr<IConcertSyncClient> ConcertSyncClient = IConcertSyncClientModule::Get().GetClient(TEXT("MultiUser")))
+		{
+			const SIZE_T Recorders = CanRecord() + RemoteRecorders();
+			bShouldSupressBecauseOfRemoteOpen = Recorders > 1 && ConcertSyncClient->GetSequencerManager()->IsSequencerRemoteOpenEnabled();
+		}
 		ULevelSequence* LevelSequence = TakePreset->GetLevelSequence();
 		if (LevelSequence)
 		{
 			UTakeRecorderSources* Sources = LevelSequence->FindMetaData<UTakeRecorderSources>();
-			return Sources && Sources->GetSources().Num() > 0;
+			return !bShouldSupressBecauseOfRemoteOpen && Sources && Sources->GetSources().Num() > 0;
 		}
 	}
 
@@ -810,6 +817,22 @@ void FConcertTakeRecorderManager::OnSessionConnectionChanged(IConcertClientSessi
 	}
 }
 
+void FConcertTakeRecorderManager::ReportRecordingError(FText &OutputError)
+{
+	const SIZE_T Recorders = CanRecord() + RemoteRecorders();
+	if(WeakSession.IsValid() && IsTakeSyncEnabled() && Recorders > 1)
+	{
+		if (TSharedPtr<IConcertSyncClient> ConcertSyncClient = IConcertSyncClientModule::Get().GetClient(TEXT("MultiUser")))
+		{
+			const bool bIsRemoteOpenEnabled = ConcertSyncClient->GetSequencerManager()->IsSequencerRemoteOpenEnabled();
+			if (bIsRemoteOpenEnabled)
+			{
+				OutputError = LOCTEXT("ErrorWidget_RemoteOpenEnabled", "Sequencer remote open is enabled. It is recommended you disable this when recording on multiple clients.");
+			}
+		}
+	}
+}
+
 bool FConcertTakeRecorderManager::CanAnyRecord() const
 {
 	UConcertSessionRecordSettings const* RecordSettings = GetDefault<UConcertSessionRecordSettings>();
@@ -904,6 +927,29 @@ void FConcertTakeRecorderManager::RecordSettingChange(const FConcertClientRecord
 			FNotificationInfo Info(ErrorText);
 			Info.ExpireDuration = 5.0f;
 			FSlateNotificationManager::Get().AddNotification(Info)->SetCompletionState(SNotificationItem::CS_None);
+
+			// Add a notification to disable remote open feature.
+			if (TSharedPtr<IConcertSyncClient> ConcertSyncClient = IConcertSyncClientModule::Get().GetClient(TEXT("MultiUser")))
+			{
+				const bool bRemoteOpen = ConcertSyncClient->GetSequencerManager()->IsSequencerRemoteOpenEnabled();
+				if (bRemoteOpen)
+				{
+					FText ShouldDisableErrorText = LOCTEXT("ShouldDisableSequencerRemoteOpen",
+											  "When multiple machines are recording, it is recommended that you disable sequencer remote open.");
+					FNotificationInfo SequencerOpenInfo(ShouldDisableErrorText);
+					SequencerOpenInfo.ExpireDuration = 10.0f;
+					SequencerOpenInfo.HyperlinkText = LOCTEXT("DisableRemoteOpen", "Disable sequencer remote open.");
+					SequencerOpenInfo.Hyperlink = FSimpleDelegate::CreateLambda([this, ConcertSyncClient]() {
+						if (ConcertSyncClient && ConcertSyncClient->GetSequencerManager())
+						{
+							ConcertSyncClient->GetSequencerManager()->SetSequencerRemoteOpen(false);
+						}
+					});
+
+					FSlateNotificationManager::Get().AddNotification(SequencerOpenInfo);
+				} 
+			}
+
 			bHaveWarned = true;
 		}
 
diff --git a/Engine/Plugins/VirtualProduction/MultiUserTakes/Source/ConcertTakeRecorder/Private/ConcertTakeRecorderManager.h b/Engine/Plugins/VirtualProduction/MultiUserTakes/Source/ConcertTakeRecorder/Private/ConcertTakeRecorderManager.h
--- a/Engine/Plugins/VirtualProduction/MultiUserTakes/Source/ConcertTakeRecorder/Private/ConcertTakeRecorderManager.h
+++ b/Engine/Plugins/VirtualProduction/MultiUserTakes/Source/ConcertTakeRecorder/Private/ConcertTakeRecorderManager.h
@@ -136,6 +136,7 @@ private:
 		return TEXT("FConcertTakeRecorderManager");
 	}
 
+	void ReportRecordingError(FText &);
 	bool CanRecord() const;
 	bool CanAnyRecord() const;
 
