horde issue #1080219 - add a null check prior before checking the sequence manager for remote open status.

#rnx

diff --git a/Engine/Plugins/VirtualProduction/MultiUserTakes/Source/ConcertTakeRecorder/Private/ConcertTakeRecorderManager.cpp b/Engine/Plugins/VirtualProduction/MultiUserTakes/Source/ConcertTakeRecorder/Private/ConcertTakeRecorderManager.cpp
--- a/Engine/Plugins/VirtualProduction/MultiUserTakes/Source/ConcertTakeRecorder/Private/ConcertTakeRecorderManager.cpp
+++ b/Engine/Plugins/VirtualProduction/MultiUserTakes/Source/ConcertTakeRecorder/Private/ConcertTakeRecorderManager.cpp
@@ -931,7 +931,7 @@ void FConcertTakeRecorderManager::RecordSettingChange(const FConcertClientRecord
 			// Add a notification to disable remote open feature.
 			if (TSharedPtr<IConcertSyncClient> ConcertSyncClient = IConcertSyncClientModule::Get().GetClient(TEXT("MultiUser")))
 			{
-				const bool bRemoteOpen = ConcertSyncClient->GetSequencerManager()->IsSequencerRemoteOpenEnabled();
+				const bool bRemoteOpen = ConcertSyncClient->GetSequencerManager() && ConcertSyncClient->GetSequencerManager()->IsSequencerRemoteOpenEnabled();
 				if (bRemoteOpen)
 				{
 					FText ShouldDisableErrorText = LOCTEXT("ShouldDisableSequencerRemoteOpen",
