When recording to a plug-in path, it is possible that the _temp package generated for the user to get captured by Multi-user. This package should be filtered.  This change checks the resolved path location to ensure that it is not included in the Multi-user ledger.

#Synced-CL 0

diff --git a/Engine/Plugins/VirtualProduction/MultiUserTakes/Source/ConcertTakeRecorder/Private/ConcertTakeRecorderManager.cpp b/Engine/Plugins/VirtualProduction/MultiUserTakes/Source/ConcertTakeRecorder/Private/ConcertTakeRecorderManager.cpp
--- a/Engine/Plugins/VirtualProduction/MultiUserTakes/Source/ConcertTakeRecorder/Private/ConcertTakeRecorderManager.cpp
+++ b/Engine/Plugins/VirtualProduction/MultiUserTakes/Source/ConcertTakeRecorder/Private/ConcertTakeRecorderManager.cpp
@@ -57,6 +57,9 @@ static FAutoConsoleVariableRef  CVarEnableTakePresetPathSync(TEXT("Concert.UseTa
 static int32 bEnableSkipHotReloadHint = 0;
 static FAutoConsoleVariableRef  CVarEnableHotReloadHint(TEXT("Concert.TakeRecorderSkipHotReloadHint"), bEnableSkipHotReloadHint, TEXT("Should we indicate to the clients that hot reload can be skipped for assets generated via take recorder."));
 
+static int32 bSkipPackageFilter = 0;
+static FAutoConsoleVariableRef  CVarSkipPackageFilter(TEXT("Concert.TakeRecorderSkipPackageFilter"), bSkipPackageFilter, TEXT("Indicates if we should skip running the package filter for take recorder while in a Multi-user session. This will force all packages to be use default behavior and _temp packages may not get excluded."));
+
 #define LOCTEXT_NAMESPACE "ConcertTakeRecorder"
 
 ITakeRecorderModule& FTakeRecorderRecorderManagerGetModule()
@@ -333,17 +336,17 @@ bool FConcertTakeRecorderManager::ShouldIconBeVisible() const
 
 	if (WeakSession.IsValid() && TakePreset && !bIsRecording && TakeSync->bSyncTakeRecordingTransactions && IsTakeSyncEnabled())
 	{
-		bool bShouldSupressBecauseOfRemoteOpen = false;
+		bool bShouldSuppressBecauseOfRemoteOpen = false;
 		if (TSharedPtr<IConcertSyncClient> ConcertSyncClient = IConcertSyncClientModule::Get().GetClient(TEXT("MultiUser")))
 		{
 			const SIZE_T Recorders = CanRecord() + RemoteRecorders();
-			bShouldSupressBecauseOfRemoteOpen = Recorders > 1 && ConcertSyncClient->GetSequencerManager()->IsSequencerRemoteOpenEnabled();
+			bShouldSuppressBecauseOfRemoteOpen = Recorders > 1 && ConcertSyncClient->GetSequencerManager()->IsSequencerRemoteOpenEnabled();
 		}
 		ULevelSequence* LevelSequence = TakePreset->GetLevelSequence();
 		if (LevelSequence)
 		{
 			UTakeRecorderSources* Sources = LevelSequence->FindMetaData<UTakeRecorderSources>();
-			return !bShouldSupressBecauseOfRemoteOpen && Sources && Sources->GetSources().Num() > 0;
+			return !bShouldSuppressBecauseOfRemoteOpen && Sources && Sources->GetSources().Num() > 0;
 		}
 	}
 
@@ -1019,20 +1022,23 @@ namespace UE::TakeRecorderManager::Private
 	const FString TempPostfix = TEXT("_temp");
 }
 
-
 EPackageFilterResult FConcertTakeRecorderManager::ShouldPackageBeFiltered(const FConcertPackageInfo& InPackageInfo)
 {
-	if (IsTakeSyncEnabled() && WeakSession.IsValid() && !CanRecord())
+	bool bShouldFilter = false;
+	FString FullName = InPackageInfo.PackageName.ToString();
+	if (IsTakeSyncEnabled() && WeakSession.IsValid() && !CanRecord() && bSkipPackageFilter == 0)
 	{
 		FTakeRecorderProjectParameters Project = GetDefault<UTakeRecorderProjectSettings>()->Settings;
-		FString FullName = InPackageInfo.PackageName.ToString();
-		if (FullName.Contains(Project.RootTakeSaveDir.Path) &&
-			FullName.Contains(UE::TakeRecorderManager::Private::TempPostfix))
+		if (  (  FullName.Contains(Project.RootTakeSaveDir.Path)
+			  || FullName.Contains(Project.GetResolvedRootTakeSaveDir()))
+			&& FullName.Contains(UE::TakeRecorderManager::Private::TempPostfix))
 		{
-			return EPackageFilterResult::Exclude;
+			bShouldFilter = true;
 		}
 	}
-	return EPackageFilterResult::UseDefault;
+	FString YesOrNo = bShouldFilter ? TEXT("Yes") : TEXT("UseDefault");
+	UE_LOG(LogConcertTakeRecorder, Verbose, TEXT("Considering package %s for filtering: %s"), *FullName, *YesOrNo);
+	return bShouldFilter ? EPackageFilterResult::Exclude : EPackageFilterResult::UseDefault;
 }
 
 bool FConcertTakeRecorderManager::CanSkipHotReload(const FConcertPackageInfo& InPackageInfo)
diff --git a/Engine/Plugins/VirtualProduction/Takes/Source/TakeRecorder/Private/Recorder/TakeRecorderParameters.cpp b/Engine/Plugins/VirtualProduction/Takes/Source/TakeRecorder/Private/Recorder/TakeRecorderParameters.cpp
--- a/Engine/Plugins/VirtualProduction/Takes/Source/TakeRecorder/Private/Recorder/TakeRecorderParameters.cpp
+++ b/Engine/Plugins/VirtualProduction/Takes/Source/TakeRecorder/Private/Recorder/TakeRecorderParameters.cpp
@@ -34,6 +34,11 @@ FTakeRecorderProjectParameters::FTakeRecorderProjectParameters()
 	, bShowNotifications(true)
 {}
 
+FString FTakeRecorderProjectParameters::GetResolvedRootTakeSaveDir() const
+{
+	return UE::TakeRecorder::Private::ResolvePathToProject(RootTakeSaveDir.Path);  
+}
+
 FString FTakeRecorderProjectParameters::GetTakeAssetPath() const
 {
 	return UE::TakeRecorder::Private::ResolvePathToProject(RootTakeSaveDir.Path) / TakeSaveDir;
diff --git a/Engine/Plugins/VirtualProduction/Takes/Source/TakeRecorder/Public/Recorder/TakeRecorderParameters.h b/Engine/Plugins/VirtualProduction/Takes/Source/TakeRecorder/Public/Recorder/TakeRecorderParameters.h
--- a/Engine/Plugins/VirtualProduction/Takes/Source/TakeRecorder/Public/Recorder/TakeRecorderParameters.h
+++ b/Engine/Plugins/VirtualProduction/Takes/Source/TakeRecorder/Public/Recorder/TakeRecorderParameters.h
@@ -74,6 +74,10 @@ struct FTakeRecorderProjectParameters
 
 	TAKERECORDER_API FTakeRecorderProjectParameters();
 
+	/** Get the resolved root take save dir. The is the path that considers both the /Game and /Plugin based directories.  */
+	TAKERECORDER_API FString GetResolvedRootTakeSaveDir() const;
+
+
 	/** The take asset path, composed of the TakeRootSaveDir and the TakeSaveDir */
 	TAKERECORDER_API FString GetTakeAssetPath() const;
 
